Fix for UE5 breakage from duplicate definition of UE_API

diff --git a/Engine/Source/Developer/Horde/Public/Horde.h b/Engine/Source/Developer/Horde/Public/Horde.h
--- a/Engine/Source/Developer/Horde/Public/Horde.h
+++ b/Engine/Source/Developer/Horde/Public/Horde.h
@@ -7,18 +7,16 @@
 
 #if WITH_EDITOR
 
-#define UE_API HORDE_API
-
 struct FHorde
 {
-	UE_API static FString GetTemplateName();
-	UE_API static FString GetTemplateId();
-	UE_API static FString GetJobId();
-	UE_API static FString GetJobURL();
-	UE_API static FString GetStepId();
-	UE_API static FString GetStepURL();
-	UE_API static FString GetStepName();
-	UE_API static FString GetBatchId();
+	HORDE_API static FString GetTemplateName();
+	HORDE_API static FString GetTemplateId();
+	HORDE_API static FString GetJobId();
+	HORDE_API static FString GetJobURL();
+	HORDE_API static FString GetStepId();
+	HORDE_API static FString GetStepURL();
+	HORDE_API static FString GetStepName();
+	HORDE_API static FString GetBatchId();
 };
 
 #endif
