Horde: Move C++ remote execution classes into the UE project.

diff --git a/Engine/Source/Developer/Horde/Private/AgentMessage.cpp b/Engine/Source/Developer/Horde/Private/AgentMessage.cpp
new file mode 100644
--- /dev/null
+++ b/Engine/Source/Developer/Horde/Private/AgentMessage.cpp
@@ -0,0 +1,329 @@
+// Copyright Epic Games, Inc. All Rights Reserved.
+
+#include "AgentMessage.h"
+#include "ComputePlatform.h"
+
+//// FAgentMessageChannel ////
+
+FAgentMessageChannel::FAgentMessageChannel(std::shared_ptr<FComputeBufferReader> InRecvBufferReader, std::shared_ptr<FComputeBufferWriter> InSendBufferWriter)
+	: RecvBufferReader(std::move(InRecvBufferReader))
+	, SendBufferWriter(std::move(InSendBufferWriter))
+	, RequestData(0)
+	, RequestSize(0)
+	, MaxRequestSize(0)
+	, ResponseType(EAgentMessageType::None)
+	, ResponseData(nullptr)
+	, ResponseLength(0)
+{
+}
+
+FAgentMessageChannel::~FAgentMessageChannel()
+{
+}
+
+void FAgentMessageChannel::Close()
+{
+	CreateMessage(EAgentMessageType::None, 0);
+	FlushMessage();
+}
+
+void FAgentMessageChannel::Ping()
+{
+	CreateMessage(EAgentMessageType::Ping, 0);
+	FlushMessage();
+}
+
+void FAgentMessageChannel::Exception(const char* Description, const char* Trace)
+{
+	CreateMessage(EAgentMessageType::Exception, strlen(Description) + strlen(Trace) + 2);
+	WriteString(Description);
+	WriteString(Trace);
+	FlushMessage();
+}
+
+void FAgentMessageChannel::Fork(int ChannelId, int BufferSize)
+{
+	CreateMessage(EAgentMessageType::Fork, sizeof(int) + sizeof(int));
+	WriteInt32(ChannelId);
+	WriteInt32(BufferSize);
+	FlushMessage();
+}
+
+void FAgentMessageChannel::Attach()
+{
+	CreateMessage(EAgentMessageType::Attach, 0);
+	FlushMessage();
+}
+
+void FAgentMessageChannel::UploadFiles(const char* Path, const char* Locator)
+{
+	CreateMessage(EAgentMessageType::WriteFiles, strlen(Path) + strlen(Locator) + 20);
+	WriteString(Path);
+	WriteString(Locator);
+	FlushMessage();
+}
+
+void FAgentMessageChannel::DeleteFiles(const char** Paths, size_t Count)
+{
+	size_t RequiredSize = 10;
+	for (size_t Idx = 0; Idx < Count; Idx++)
+	{
+		RequiredSize += strlen(Paths[Idx]) + 10;
+	}
+
+	CreateMessage(EAgentMessageType::DeleteFiles, RequiredSize);
+	WriteUnsignedVarInt(Count);
+	for (size_t Idx = 0; Idx < Count; Idx++)
+	{
+		WriteString(Paths[Idx]);
+	}
+	FlushMessage();
+}
+
+void FAgentMessageChannel::Execute(const char* Exe, const char** Args, size_t NumArgs, const char* WorkingDir, const char** EnvVars, size_t NumEnvVars, EExecuteProcessFlags Flags)
+{
+	size_t RequiredSize = 50 + strlen(Exe);
+	for (size_t Idx = 0; Idx < NumArgs; Idx++)
+	{
+		RequiredSize += strlen(Args[Idx]) + 10;
+	}
+	if (WorkingDir != nullptr)
+	{
+		RequiredSize += strlen(WorkingDir) + 10;
+	}
+	for (size_t Idx = 0; Idx < NumEnvVars; Idx++)
+	{
+		RequiredSize += strlen(EnvVars[Idx]) + 20;
+	}
+
+	CreateMessage(EAgentMessageType::ExecuteV2, RequiredSize);
+	WriteString(Exe);
+
+	WriteUnsignedVarInt(NumArgs);
+	for (size_t Idx = 0; Idx < NumArgs; Idx++)
+	{
+		WriteString(Args[Idx]);
+	}
+
+	WriteOptionalString(WorkingDir);
+
+	WriteUnsignedVarInt(NumEnvVars);
+	for (size_t Idx = 0; Idx < NumEnvVars; Idx++)
+	{
+		const char* EqualsPtr = strchr(EnvVars[Idx], '=');
+		UE_COMPUTE_ASSERT(EqualsPtr != nullptr);
+
+		WriteString(std::string_view(EnvVars[Idx], EqualsPtr - EnvVars[Idx]));
+		if (*(EqualsPtr + 1) == 0)
+		{
+			WriteOptionalString(nullptr);
+		}
+		else
+		{
+			WriteOptionalString(EqualsPtr + 1);
+		}
+	}
+
+	WriteInt32((int)Flags);
+	FlushMessage();
+}
+
+void FAgentMessageChannel::Blob(const unsigned char* Data, size_t Length)
+{
+	const size_t MaxChunkSize = 512 * 1024;
+	for (size_t ChunkOffset = 0; ChunkOffset < Length;)
+	{
+		size_t ChunkLength = std::min(Length - ChunkOffset, MaxChunkSize);
+		
+		CreateMessage(EAgentMessageType::ReadBlobResponse, ChunkLength + 128);
+		WriteInt32((int)ChunkOffset);
+		WriteInt32((int)Length);
+		WriteFixedLengthBytes(Data + ChunkOffset, ChunkLength);
+		FlushMessage();
+
+		ChunkOffset += ChunkLength;
+	}
+}
+
+void FAgentMessageChannel::Xor(const unsigned char* Data, size_t Length, unsigned char Value)
+{
+	CreateMessage(EAgentMessageType::XorRequest, Length + 1);
+	WriteFixedLengthBytes(Data, Length);
+	WriteFixedLengthBytes(&Value, 1);
+	FlushMessage();
+}
+
+EAgentMessageType FAgentMessageChannel::ReadResponse()
+{
+	if (ResponseLength > 0)
+	{
+		RecvBufferReader->AdvanceReadPosition(ResponseLength);
+		ResponseData = nullptr;
+		ResponseLength = 0;
+	}
+
+	const unsigned char* Header = RecvBufferReader->WaitToRead(MessageHeaderLength);
+	unsigned int Length = *((unsigned int*)(Header + 1));
+
+	Header = RecvBufferReader->WaitToRead(MessageHeaderLength + Length);
+
+	ResponseType = (EAgentMessageType)Header[0];
+	ResponseData = Header + MessageHeaderLength;
+	ResponseLength = Length;
+
+	return ResponseType;
+}
+
+void FAgentMessageChannel::ReadException(AgentMessage::FException& Ex)
+{
+	UE_COMPUTE_ASSERT(ResponseType == EAgentMessageType::Exception);
+
+	const unsigned char* Pos = ResponseData;
+	Ex.Message = ReadString(&Pos);
+	Ex.Description = ReadString(&Pos);
+	UE_COMPUTE_ASSERT(Pos == ResponseData + ResponseLength);
+}
+
+int FAgentMessageChannel::ReadExecuteResult()
+{
+	UE_COMPUTE_ASSERT(ResponseType == EAgentMessageType::ExecuteResult);
+
+	const unsigned char* Pos = ResponseData;
+	int Result = ReadInt32(&Pos);
+	UE_COMPUTE_ASSERT(Pos == ResponseData + ResponseLength);
+
+	return Result;
+}
+
+void FAgentMessageChannel::ReadBlobRequest(AgentMessage::FBlobRequest& Ex)
+{
+	UE_COMPUTE_ASSERT(ResponseType == EAgentMessageType::ReadBlob);
+
+	const unsigned char* Pos = ResponseData;
+	Ex.Locator = ReadString(&Pos);
+	Ex.Offset = (int)ReadUnsignedVarInt(&Pos);
+	Ex.Length = (int)ReadUnsignedVarInt(&Pos);
+	UE_COMPUTE_ASSERT(Pos == ResponseData + ResponseLength);
+}
+
+void FAgentMessageChannel::CreateMessage(EAgentMessageType Type, size_t MaxLength)
+{
+	RequestData = SendBufferWriter->WaitToWrite(MessageHeaderLength + MaxLength);
+	RequestData[0] = (unsigned char)Type;
+}
+
+void FAgentMessageChannel::FlushMessage()
+{
+	memcpy(&RequestData[1], &RequestSize, sizeof(int));
+	SendBufferWriter->AdvanceWritePosition(MessageHeaderLength + RequestSize);
+	RequestSize = 0;
+	MaxRequestSize = 0;
+	RequestData = nullptr;
+}
+
+void FAgentMessageChannel::WriteInt32(int Value)
+{
+	WriteFixedLengthBytes((const unsigned char*)&Value, sizeof(int));
+}
+
+int FAgentMessageChannel::ReadInt32(const unsigned char** Pos)
+{
+	int Value;
+	memcpy(&Value, *Pos, sizeof(int));
+	*Pos += sizeof(int);
+	return Value;
+}
+
+void FAgentMessageChannel::WriteFixedLengthBytes(const unsigned char* Data, size_t Length)
+{
+	UE_COMPUTE_ASSERT(RequestSize + Length <= MaxRequestSize);
+	memcpy(&RequestData[MessageHeaderLength + RequestSize], Data, Length);
+	RequestSize += Length;
+}
+
+const unsigned char* FAgentMessageChannel::ReadFixedLengthBytes(const unsigned char** Pos, size_t Length)
+{
+	const unsigned char* Data = *Pos;
+	*Pos += Length;
+	return Data;
+}
+
+size_t FAgentMessageChannel::MeasureUnsignedVarInt(size_t Value)
+{
+	UE_COMPUTE_ASSERT(Value == (unsigned int)Value);
+	return (FComputePlatform::FloorLog2((unsigned int)Value) / 7) + 1;
+}
+
+void FAgentMessageChannel::WriteUnsignedVarInt(size_t Value)
+{
+	size_t ByteCount = MeasureUnsignedVarInt(Value);
+	UE_COMPUTE_ASSERT(RequestSize + ByteCount <= MaxRequestSize);
+
+	unsigned char* Output = RequestData + MessageHeaderLength + RequestSize;
+	for (size_t Idx = 1; Idx < ByteCount; Idx++)
+	{
+		Output[ByteCount - Idx] = (unsigned)Value;
+		Value >>= 8;
+	}
+	Output[0] = (unsigned char)((0xff << (9 - (int)ByteCount)) | (unsigned char)Value);
+
+	RequestSize += ByteCount;
+}
+
+size_t FAgentMessageChannel::ReadUnsignedVarInt(const unsigned char** Pos)
+{
+	const unsigned char* Data = *Pos;
+
+	unsigned char FirstByte = Data[0];
+	size_t NumBytes = FComputePlatform::CountLeadingZeros(~(unsigned int)FirstByte) + 1 - 24; // Note byte -> int conversion here, hence ignoring subtracting 24 bits
+
+	size_t value = (size_t)(FirstByte & (0xff >> NumBytes));
+	for (size_t Idx = 1; Idx < NumBytes; Idx++)
+	{
+		value <<= 8;
+		value |= Data[Idx];
+	}
+
+	*Pos += NumBytes;
+	return value;
+}
+
+size_t FAgentMessageChannel::MeasureString(const char* Text) const
+{
+	size_t Length = strlen(Text);
+	return MeasureUnsignedVarInt(Length) + Length;
+}
+
+void FAgentMessageChannel::WriteString(const char* Text)
+{
+	size_t Length = strlen(Text);
+	WriteUnsignedVarInt(Length);
+	WriteFixedLengthBytes((const unsigned char*)Text, Length);
+}
+
+void FAgentMessageChannel::WriteString(const std::string_view& Text)
+{
+	WriteUnsignedVarInt(Text.size());
+	WriteFixedLengthBytes((const unsigned char*)Text.data(), Text.size());
+}
+
+std::string_view FAgentMessageChannel::ReadString(const unsigned char** Pos)
+{
+	size_t Length = ReadUnsignedVarInt(Pos);
+	const char* Start = (const char*)ReadFixedLengthBytes(Pos, Length);
+	return std::string_view(Start, Length);
+}
+
+void FAgentMessageChannel::WriteOptionalString(const char* Text)
+{
+	if (Text == nullptr)
+	{
+		WriteUnsignedVarInt(0);
+	}
+	else
+	{
+		size_t Length = strlen(Text);
+		WriteUnsignedVarInt(Length + 1);
+		WriteFixedLengthBytes((const unsigned char*)Text, Length);
+	}
+}
diff --git a/Engine/Source/Developer/Horde/Private/ComputeBuffer.cpp b/Engine/Source/Developer/Horde/Private/ComputeBuffer.cpp
new file mode 100644
--- /dev/null
+++ b/Engine/Source/Developer/Horde/Private/ComputeBuffer.cpp
@@ -0,0 +1,1188 @@
+// Copyright Epic Games, Inc. All Rights Reserved.
+
+#include "ComputeBuffer.h"
+#include "ComputePlatform.h"
+#include <assert.h>
+#include <iostream>
+#include <assert.h>
+#include <stdio.h>
+#include <string.h>
+
+struct FComputeBufferDetail
+{
+	static const int HeaderSize = (2 + FComputeBuffer::MaxReaders + FComputeBuffer::MaxChunks) * sizeof(long long);
+
+	enum class EWriteState : unsigned long long
+	{
+		// Writer has moved to the next chunk
+		MovedToNext = 0,
+
+		// Chunk is still being appended to
+		Writing = 2,
+
+		// This chunk marks the end of the stream
+		Complete = 3,
+	};
+
+	union FChunkState
+	{
+	public:
+		long long Value;
+
+		FChunkState()
+			: Value(0)
+		{ }
+
+		FChunkState(long long InValue)
+			: Value(InValue)
+		{ }
+
+		// Constructor
+		FChunkState(EWriteState WriteState, int ReaderFlags, int Length)
+			: Value((unsigned long long)Length | ((unsigned long long)ReaderFlags << 31) | ((unsigned long long)WriteState << 62))
+		{ }
+
+		// Written length of this chunk
+		int GetLength() const { return (int)(Value & 0x7fffffff); }
+
+		// Set of flags which are set for each reader that still has to read from a chunk
+		int GetReaderFlags() const { return (int)((Value >> 31) & 0x7fffffff); }
+
+		// State of the writer
+		EWriteState GetWriteState() const { return (EWriteState)((unsigned long long)Value >> 62); }
+
+		// Test whether a particular reader is still referencing the chunk
+		bool HasReaderFlag(int ReaderIdx) const { return (Value & (1ULL << (31 + ReaderIdx))) != 0; }
+
+	private:
+		struct // For debugging purposes only; non-portable assumption of bitfield layout
+		{
+			unsigned long long Length : 31;
+			unsigned long long ReaderFlags : 31;
+			EWriteState WriteState : 2;
+		} Fields;
+	};
+
+	struct FChunkStatePtr
+	{
+	public:
+		FChunkStatePtr(volatile FChunkState* InChunkState)
+			: State(InChunkState)
+		{ }
+
+		// Read the state value from memory
+		FChunkState Get() const
+		{
+			return FChunkState(FComputePlatform::AtomicRead64(&State->Value));
+		}
+
+		// Read the state value from memory
+		void Set(FChunkState NewState)
+		{
+			FComputePlatform::AtomicWrite64(&State->Value, NewState.Value);
+		}
+
+		// Attempt to update the chunk state
+		bool TryUpdate(FChunkState PrevState, FChunkState NextState)
+		{
+			return FComputePlatform::AtomicCompareExchange64(&State->Value, NextState.Value, PrevState.Value);
+		}
+
+		// Append data to the chunk
+		void Append(long long Length)
+		{
+			FComputePlatform::AtomicAdd64(&State->Value, Length);
+		}
+
+		// Mark this chunk as the end of the stream
+		void MarkComplete()
+		{
+			FComputePlatform::AtomicOr64(&State->Value, FChunkState(EWriteState::Complete, 0, 0).Value);
+		}
+
+		// Start reading the chunk with the given reader
+		void StartReading(int ReaderIdx)
+		{
+			FComputePlatform::AtomicOr64(&State->Value, FChunkState((EWriteState)0, 1 << ReaderIdx, 0).Value);
+		}
+
+		// Clear the reader flag
+		void FinishReading(int ReaderIdx)
+		{
+			FComputePlatform::AtomicAnd64(&State->Value, ~FChunkState((EWriteState)0, 1 << ReaderIdx, 0).Value);
+		}
+
+		// Move to the next chunk
+		void FinishWriting()
+		{
+			FComputePlatform::AtomicAnd64(&State->Value, ~FChunkState(EWriteState::Writing, 0, 0).Value);
+		}
+
+	private:
+		volatile FChunkState* const State;
+	};
+
+	union FReaderState
+	{
+		long long Value;
+
+		FReaderState()
+			: Value(0)
+		{ }
+
+		FReaderState(long long InValue)
+			: Value(InValue)
+		{ }
+
+		FReaderState(int ChunkIdx, int Offset, int RefCount, bool Detached)
+			: Value((long long)(unsigned int)Offset | ((long long)(unsigned int)ChunkIdx << 32) | ((long long)(unsigned int)RefCount << 40) | (long long)(Detached ? (1ULL << 63) : 0))
+		{ }
+
+		int GetOffset() const
+		{
+			return (int)(Value & 0xffffffff);
+		}
+
+		int GetChunkIdx() const
+		{
+			return (int)((Value >> 32) & 0xff);
+		}
+
+		int GetRefCount() const
+		{
+			return (int)((Value >> 40) & 0x7fff);
+		}
+
+		bool IsDetached() const
+		{
+			return (Value & (1ULL << 63)) != 0;
+		}
+
+	private:
+		struct // For debugging purposes only; non-portable assumption of bitfield layout
+		{
+			unsigned long long Offset : 32;
+			unsigned long long ChunkIdx : 8;
+			unsigned long long RefCount : 23;
+			unsigned long long Detached : 1;
+		} Fields;
+	};
+
+	struct FReaderStatePtr
+	{
+	public:
+		FReaderStatePtr(volatile FReaderState* InReaderState)
+			: State(InReaderState)
+		{ }
+
+		// Read the state value from memory
+		FReaderState Get() const
+		{
+			return FReaderState(FComputePlatform::AtomicRead64(&State->Value));
+		}
+
+		// Read the state value from memory
+		void Set(FReaderState NewState)
+		{
+			FComputePlatform::AtomicWrite64(&State->Value, NewState.Value);
+		}
+
+		// Attempt to update the chunk state
+		bool TryUpdate(FReaderState PrevState, FReaderState NextState)
+		{
+			return FComputePlatform::AtomicCompareExchange64(&State->Value, NextState.Value, PrevState.Value);
+		}
+
+	private:
+		volatile FReaderState* const State;
+	};
+
+	union FWriterState
+	{
+		long long Value;
+
+		FWriterState()
+			: Value(0)
+		{ }
+
+		FWriterState(long long InValue) 
+			: Value(InValue) 
+		{ }
+		
+		FWriterState(int ChunkIdx, int ReaderFlags, int RefCount, bool HasWrapped) 
+			: Value((long long)(unsigned int)ChunkIdx | ((long long)(unsigned int)ReaderFlags << 32) | ((long long)(unsigned int)RefCount << 48) | (HasWrapped ? (1ULL << 63) : 0))
+		{ }
+
+		int GetChunkIdx() const { return Value & 0x7fffffff; }
+		int GetReaderFlags() const { return ((unsigned long long)Value >> 32) & 0xffff; }
+		int GetRefCount() const { return ((unsigned long long)Value >> 48) & 0x7fff; }
+		bool HasWrapped() const { return (Value & (1ULL << 63)) != 0; }
+
+	private:
+		struct // For debugging purposes only; non-portable assumption of bitfield layout
+		{
+			int ChunkIdx : 32;
+			int ReaderFlags : 8;
+			int RefCount : 23;
+			unsigned int Wrapped : 1;
+		} Fields;
+	};
+
+	struct FWriterStatePtr
+	{
+	public:
+		FWriterStatePtr(volatile FWriterState* InState)
+			: State(InState)
+		{ }
+
+		FWriterState Get() const
+		{
+			return FWriterState(FComputePlatform::AtomicRead64(&State->Value));
+		}
+
+		void Set(FWriterState NewState)
+		{
+			FComputePlatform::AtomicWrite64(&State->Value, NewState.Value);
+		}
+
+		bool TryUpdate(FWriterState PrevValue, FWriterState NextValue)
+		{
+			return FComputePlatform::AtomicCompareExchange64(&State->Value, NextValue.Value, PrevValue.Value);
+		}
+
+	private:
+		volatile FWriterState* State;
+	};
+
+	struct FHeader
+	{
+	public:
+		const unsigned short int NumReaders;
+		const unsigned short int NumChunks;
+		const int ChunkLength;
+
+		FHeader(int InNumReaders, int InNumChunks, int InChunkLength)
+			: NumReaders(InNumReaders)
+			, NumChunks(InNumChunks)
+			, ChunkLength(InChunkLength)
+		{
+			static_assert(sizeof(FWriterState) == sizeof(long long), "Incorrect size of FWriterState");
+			static_assert(sizeof(FReaderState) == sizeof(long long), "Incorrect size of FReaderState");
+			static_assert(sizeof(FChunkState) == sizeof(long long), "Incorrect size of FChunkState");
+			static_assert(sizeof(FHeader) == HeaderSize, "Incorrect size of FHeader");
+
+			GetChunkStatePtr(0).Set(FChunkState(EWriteState::Writing, 0, 0));
+		}
+
+		FWriterStatePtr GetWriterStatePtr()
+		{
+			return FWriterStatePtr(&Writer);
+		}
+
+		FReaderStatePtr GetReaderStatePtr(int ReaderIdx)
+		{
+			return FReaderStatePtr(&Readers[ReaderIdx]);
+		}
+
+		FChunkStatePtr GetChunkStatePtr(int ChunkIdx)
+		{
+			return FChunkStatePtr(&Chunks[ChunkIdx]);
+		}
+
+	private:
+		FWriterState Writer;
+		FReaderState Readers[FComputeBuffer::MaxReaders];
+		FChunkState Chunks[FComputeBuffer::MaxChunks];
+	};
+
+	char Name[FComputeBuffer::MaxNameLength];
+
+	FHeader* Header;
+	unsigned char* ChunkPtrs[FComputeBuffer::MaxChunks];
+
+	FComputeEvent WriterEvent;
+	FComputeEvent ReaderEvents[FComputeBuffer::MaxReaders];
+
+	long RefCount;
+
+	FComputeBufferDetail(const char* InName)
+		: Header(nullptr)
+		, ChunkPtrs{ nullptr, }
+		, RefCount(1)
+	{
+		static_assert(sizeof(FChunkStatePtr) == sizeof(long long), "Incorrect size of FChunkStatePtr; check union is declared correctly.");
+		static_assert(sizeof(FWriterStatePtr) == sizeof(long long), "Incorrect size of FWriterStatePtr; check union is declared correctly.");
+
+		FComputePlatform::Strcpy(Name, FComputeBuffer::MaxNameLength, InName);
+	}
+
+	~FComputeBufferDetail()
+	{
+		Header = nullptr;
+	}
+
+	void AddRef()
+	{
+		FComputePlatform::AtomicIncrement(&RefCount);
+	}
+
+	void Release()
+	{
+		if (FComputePlatform::AtomicDecrement(&RefCount) == 0)
+		{
+			delete this;
+		}
+	}
+
+	static std::unique_ptr<FComputeBufferDetail> CreateNew(const FComputeBuffer::FParams& Params)
+	{
+		long long Capacity = sizeof(FHeader) + (Params.NumChunks * sizeof(unsigned int)) + (Params.NumChunks * Params.ChunkLength);
+
+		const char* Name = Params.Name;
+
+		char BaseNameBuffer[FComputeBuffer::MaxNameLength];
+		if (Name == nullptr)
+		{
+			FComputePlatform::CreateUniqueName(BaseNameBuffer, FComputeBuffer::MaxNameLength);
+			Name = BaseNameBuffer;
+		}
+
+		char NameBuffer[FComputeBuffer::MaxNameLength];
+		snprintf(NameBuffer, FComputeBuffer::MaxNameLength, "%s_M", Name);
+
+		std::unique_ptr<FComputeBufferDetail> Detail = std::make_unique<FComputeBufferDetail>(Name);
+		if (!Detail->MemoryMappedFile.Create(NameBuffer, Capacity))
+		{
+			return nullptr;
+		}
+
+		void* Pointer = Detail->MemoryMappedFile.GetPointer();
+		if (Pointer == nullptr)
+		{
+			return nullptr;
+		}
+
+		Detail->Header = new(Pointer) FHeader(Params.NumReaders, Params.NumChunks, Params.ChunkLength);
+
+		snprintf(NameBuffer, FComputeBuffer::MaxNameLength, "%s_W", Name);
+		if (!Detail->WriterEvent.Create(NameBuffer))
+		{
+			return nullptr;
+		}
+
+		for (int ReaderIdx = 0; ReaderIdx < Detail->Header->NumReaders; ReaderIdx++)
+		{
+			snprintf(NameBuffer, FComputeBuffer::MaxNameLength, "%s_R%d", Name, ReaderIdx);
+			if (!Detail->ReaderEvents[ReaderIdx].Create(NameBuffer))
+			{
+				return nullptr;
+			}
+		}
+
+		InitShared(Detail.get());
+		return Detail;
+	}
+
+	static std::unique_ptr<FComputeBufferDetail> OpenExisting(const char* Name)
+	{
+		char NameBuffer[FComputeBuffer::MaxNameLength];
+		snprintf(NameBuffer, FComputeBuffer::MaxNameLength, "%s_M", Name);
+
+		std::unique_ptr<FComputeBufferDetail> Detail = std::make_unique<FComputeBufferDetail>(Name);
+		if (!Detail->MemoryMappedFile.OpenExisting(NameBuffer))
+		{
+			return nullptr;
+		}
+
+		Detail->Header = (FHeader*)Detail->MemoryMappedFile.GetPointer();
+		if (Detail->Header == nullptr)
+		{
+			return nullptr;
+		}
+
+		snprintf(NameBuffer, FComputeBuffer::MaxNameLength, "%s_W", Name);
+		if (!Detail->WriterEvent.OpenExisting(NameBuffer))
+		{
+			return nullptr;
+		}
+
+		for (int ReaderIdx = 0; ReaderIdx < Detail->Header->NumReaders; ReaderIdx++)
+		{
+			snprintf(NameBuffer, FComputeBuffer::MaxNameLength, "%s_R%d", Name, ReaderIdx);
+			if (!Detail->ReaderEvents[ReaderIdx].OpenExisting(NameBuffer))
+			{
+				return nullptr;
+			}
+		}
+
+		InitShared(Detail.get());
+		return Detail;
+	}
+
+	int CreateReader()
+	{
+		for(int ReaderIdx = 0; ReaderIdx < Header->NumReaders; ReaderIdx++)
+		{
+			FReaderStatePtr ReaderStatePtr = Header->GetReaderStatePtr(ReaderIdx);
+			for (;;)
+			{
+				FReaderState ReaderState = ReaderStatePtr.Get();
+				if (ReaderState.GetRefCount() > 0)
+				{
+					break;
+				}
+				if (ReaderStatePtr.TryUpdate(ReaderState, FReaderState(0, 0, 1, false)))
+				{
+					FWriterStatePtr WriterStatePtr = Header->GetWriterStatePtr();
+					for (;;)
+					{
+						FWriterState WriterState = WriterStatePtr.Get();
+						UE_COMPUTE_ASSERT(!WriterState.HasWrapped());
+
+						if (WriterStatePtr.TryUpdate(WriterState, FWriterState(WriterState.GetChunkIdx(), WriterState.GetReaderFlags() | (1 << ReaderIdx), WriterState.GetRefCount(), WriterState.HasWrapped())))
+						{
+							for (int WriteChunkIdx = 0; WriteChunkIdx <= WriterState.GetChunkIdx(); WriteChunkIdx++)
+							{
+								Header->GetChunkStatePtr(WriteChunkIdx).StartReading(ReaderIdx);
+							}
+							return ReaderIdx;
+						}
+					}
+				}
+			}
+		}
+		return -1;
+	}
+
+	void AddReaderRef(int ReaderIdx)
+	{
+		FReaderStatePtr ReaderStatePtr = Header->GetReaderStatePtr(ReaderIdx);
+		for (; ; )
+		{
+			FReaderState ReaderState = ReaderStatePtr.Get();
+			UE_COMPUTE_ASSERT(ReaderState.GetRefCount() > 0);
+
+			if (ReaderStatePtr.TryUpdate(ReaderState, FReaderState(ReaderState.GetChunkIdx(), ReaderState.GetOffset(), ReaderState.GetRefCount() + 1, ReaderState.IsDetached())))
+			{
+				break;
+			}
+		}
+	}
+
+	void ReleaseReaderRef(int ReaderIdx)
+	{
+		FReaderStatePtr ReaderStatePtr = Header->GetReaderStatePtr(ReaderIdx);
+		for (; ; )
+		{
+			FReaderState ReaderState = ReaderStatePtr.Get();
+			UE_COMPUTE_ASSERT(ReaderState.GetRefCount() > 0);
+
+			if (ReaderState.GetRefCount() == 1)
+			{
+				for (int Idx = 0; Idx < Header->NumChunks; Idx++)
+				{
+					Header->GetChunkStatePtr(Idx).FinishReading(ReaderIdx);
+				}
+			}
+
+			if (ReaderStatePtr.TryUpdate(ReaderState, FReaderState(ReaderState.GetChunkIdx(), ReaderState.GetOffset(), ReaderState.GetRefCount() - 1, ReaderState.IsDetached())))
+			{
+				break;
+			}
+		}
+	}
+
+	void CreateWriter()
+	{
+		FWriterStatePtr WriterStatePtr = Header->GetWriterStatePtr();
+		for (; ; )
+		{
+			FWriterState WriterState = WriterStatePtr.Get();
+			UE_COMPUTE_ASSERT(WriterState.GetRefCount() == 0);
+
+			if (WriterStatePtr.TryUpdate(WriterState, FWriterState(WriterState.GetChunkIdx(), WriterState.GetReaderFlags(), 1, WriterState.HasWrapped())))
+			{
+				FChunkStatePtr ChunkStatePtr = Header->GetChunkStatePtr(WriterState.GetChunkIdx());
+				for (; ; )
+				{
+					FChunkState ChunkState = ChunkStatePtr.Get();
+					if (ChunkStatePtr.TryUpdate(ChunkState, FChunkState(EWriteState::Writing, ChunkState.GetReaderFlags(), ChunkState.GetLength())))
+					{
+						break;
+					}
+				}
+				break;
+			}
+		}
+	}
+
+	void AddWriterRef()
+	{
+		FWriterStatePtr WriterStatePtr = Header->GetWriterStatePtr();
+		for (; ; )
+		{
+			FWriterState WriterState = WriterStatePtr.Get();
+			UE_COMPUTE_ASSERT(WriterState.GetRefCount() > 0);
+
+			if (WriterStatePtr.TryUpdate(WriterState, FWriterState(WriterState.GetChunkIdx(), WriterState.GetReaderFlags(), WriterState.GetRefCount() + 1, WriterState.HasWrapped())))
+			{
+				break;
+			}
+		}
+	}
+
+	void ReleaseWriterRef()
+	{
+		FWriterStatePtr writerStatePtr = Header->GetWriterStatePtr();
+		for (; ; )
+		{
+			FWriterState writerState = writerStatePtr.Get();
+			UE_COMPUTE_ASSERT(writerState.GetRefCount() > 0);
+
+			if (writerState.GetRefCount() == 1)
+			{
+				MarkComplete();
+			}
+
+			if (writerStatePtr.TryUpdate(writerState, FWriterState(writerState.GetChunkIdx(), writerState.GetReaderFlags(), writerState.GetRefCount() - 1, writerState.HasWrapped())))
+			{
+				break;
+			}
+		}
+	}
+
+	bool IsComplete(int ReaderIdx) const
+	{
+		FReaderState ReaderState = Header->GetReaderStatePtr(ReaderIdx).Get();
+		if (ReaderState.IsDetached())
+		{
+			return true;
+		}
+
+		FChunkState ChunkState = Header->GetChunkStatePtr(ReaderState.GetChunkIdx()).Get();
+		return ChunkState.GetWriteState() == EWriteState::Complete && ReaderState.GetOffset() == ChunkState.GetLength();
+	}
+
+	void DetachReader(int ReaderIdx)
+	{
+		FReaderStatePtr ReaderStatePtr = Header->GetReaderStatePtr(ReaderIdx);
+		for (; ; )
+		{
+			FReaderState ReaderState = ReaderStatePtr.Get();
+			if (ReaderStatePtr.TryUpdate(ReaderState, FReaderState(ReaderState.GetChunkIdx(), ReaderState.GetOffset(), ReaderState.GetRefCount(), true)))
+			{
+				ReaderEvents[ReaderIdx].Signal();
+				break;
+			}
+		}
+	}
+
+	void AdvanceReadPosition(int ReaderIdx, size_t Offset)
+	{
+		FReaderStatePtr ReaderStatePtr = Header->GetReaderStatePtr(ReaderIdx);
+		for (; ; )
+		{
+			FReaderState ReaderState = ReaderStatePtr.Get();
+			if (ReaderStatePtr.TryUpdate(ReaderState, FReaderState(ReaderState.GetChunkIdx(), ReaderState.GetOffset() + (int)Offset, ReaderState.GetRefCount(), ReaderState.IsDetached())))
+			{
+				ReaderEvents[ReaderIdx].Signal();
+				break;
+			}
+		}
+	}
+
+	size_t GetMaxReadSize(int ReaderIdx) const
+	{
+		FReaderState ReaderState = Header->GetReaderStatePtr(ReaderIdx).Get();
+		if (ReaderState.IsDetached())
+		{
+			return 0;
+		}
+
+		FChunkStatePtr ChunkStatePtr = Header->GetChunkStatePtr(ReaderState.GetChunkIdx());
+		FChunkState ChunkState = ChunkStatePtr.Get();
+
+		if (ChunkState.HasReaderFlag(ReaderIdx))
+		{
+			return ChunkState.GetLength() - ReaderState.GetOffset();
+		}
+		else
+		{
+			return 0;
+		}
+	}
+
+	const unsigned char* WaitToRead(int ReaderIdx, size_t MinSize, int TimeoutMs)
+	{
+		for (; ; )
+		{
+			FReaderStatePtr ReaderStatePtr = Header->GetReaderStatePtr(ReaderIdx);
+			FReaderState ReaderState = ReaderStatePtr.Get();
+			if (ReaderState.IsDetached())
+			{
+				return nullptr;
+			}
+
+			FChunkStatePtr ChunkStatePtr = Header->GetChunkStatePtr(ReaderState.GetChunkIdx());
+			FChunkState ChunkState = ChunkStatePtr.Get();
+
+			if (!ChunkState.HasReaderFlag(ReaderIdx))
+			{
+				// Wait until the current chunk is readable
+				if (!ReaderEvents[ReaderIdx].Wait(TimeoutMs))
+				{
+					return nullptr;
+				}
+			}
+			else if (ReaderState.GetOffset() + MinSize <= ChunkState.GetLength())
+			{
+				// We have enough data in the chunk to be able to read a message
+				return ChunkPtrs[ReaderState.GetChunkIdx()] + ReaderState.GetOffset();
+			}
+			else if (ChunkState.GetWriteState() == EWriteState::Writing)
+			{
+				// Wait until there is more data in the chunk
+				if (!ReaderEvents[ReaderIdx].Wait(TimeoutMs))
+				{
+					return nullptr;
+				}
+			}
+			else if (ReaderState.GetOffset() < ChunkState.GetLength() || ChunkState.GetWriteState() == EWriteState::Complete)
+			{
+				// Cannot read the requested amount of data from this chunk.
+				return nullptr;
+			}
+			else if (ChunkState.GetWriteState() == EWriteState::MovedToNext)
+			{
+				// Move to the next chunk
+				ChunkStatePtr.FinishReading(ReaderIdx);
+				WriterEvent.Signal();
+
+				int NextChunkIdx = ReaderState.GetChunkIdx() + 1;
+				if (NextChunkIdx == Header->NumChunks)
+				{
+					NextChunkIdx = 0;
+				}
+
+				ReaderStatePtr.TryUpdate(ReaderState, FReaderState(NextChunkIdx, 0, ReaderState.GetRefCount(), ReaderState.IsDetached()));
+			}
+			else
+			{
+				UE_COMPUTE_ASSERT(false);
+			}
+		}
+	}
+
+	bool MarkComplete()
+	{
+		FWriterState WriterState = Header->GetWriterStatePtr().Get();
+
+		FChunkStatePtr ChunkStatePtr = Header->GetChunkStatePtr(WriterState.GetChunkIdx());
+		if (ChunkStatePtr.Get().GetWriteState() != EWriteState::Complete)
+		{
+			ChunkStatePtr.MarkComplete();
+			SetAllReadEvents();
+			return true;
+		}
+
+		return false;
+	}
+
+	void AdvanceWritePosition(size_t Size)
+	{
+		if (Size > 0)
+		{
+			FWriterState WriterState = Header->GetWriterStatePtr().Get();
+
+			FChunkStatePtr ChunkStatePtr = Header->GetChunkStatePtr(WriterState.GetChunkIdx());
+			FChunkState ChunkState = ChunkStatePtr.Get();
+
+			UE_COMPUTE_ASSERT(ChunkState.GetWriteState() == EWriteState::Writing);
+			ChunkStatePtr.Append(Size);
+
+			SetAllReadEvents();
+		}
+	}
+
+	size_t GetMaxWriteSize() const
+	{
+		FWriterState WriterState = Header->GetWriterStatePtr().Get();
+
+		FChunkState ChunkState = Header->GetChunkStatePtr(WriterState.GetChunkIdx()).Get();
+		if (ChunkState.GetWriteState() == EWriteState::Writing)
+		{
+			return Header->ChunkLength - ChunkState.GetLength();
+		}
+		else
+		{
+			return 0;
+		}
+	}
+
+	size_t Write(const void* Buffer, size_t MaxSize, int TimeoutMs)
+	{
+		unsigned char* SendData = WaitToWrite(1, TimeoutMs);
+		if (SendData == nullptr)
+		{
+			return 0;
+		}
+
+		size_t SendSize = GetMaxWriteSize();
+		if (MaxSize < SendSize)
+		{
+			SendSize = MaxSize;
+		}
+
+		memcpy(SendData, Buffer, SendSize);
+		AdvanceWritePosition(SendSize);
+		return SendSize;
+	}
+
+	unsigned char* WaitToWrite(size_t MinSize, int TimeoutMs)
+	{
+		UE_COMPUTE_ASSERT(MinSize <= Header->ChunkLength);
+
+		// Get the current chunk we're writing to
+		FWriterState WriterState = Header->GetWriterStatePtr().Get();
+		int WriteChunkIdx = WriterState.GetChunkIdx();
+
+		FChunkStatePtr WriteChunkStatePtr = Header->GetChunkStatePtr(WriteChunkIdx);
+
+		// Check if we can append to this chunk
+		FChunkState ChunkState = WriteChunkStatePtr.Get();
+		if (ChunkState.GetWriteState() == EWriteState::Writing)
+		{
+			int Length = ChunkState.GetLength();
+			if (Length + MinSize <= Header->ChunkLength)
+			{
+				return ChunkPtrs[WriteChunkIdx] + Length;
+			}
+
+			WriteChunkStatePtr.FinishWriting(); // STATE CHANGE
+			SetAllReadEvents();
+		}
+
+		if (ChunkState.GetWriteState() == EWriteState::Complete)
+		{
+			return nullptr;
+		}
+
+		// Otherwise get the next chunk to write to
+		int NextWriteChunkIdx = WriteChunkIdx + 1;
+		if (NextWriteChunkIdx == Header->NumChunks)
+		{
+			NextWriteChunkIdx = 0;
+		}
+
+		// Wait until all readers have finished with the chunk, and we can update the writer to match
+		FChunkStatePtr NextWriteChunkStatePtr = Header->GetChunkStatePtr(NextWriteChunkIdx);
+		for (;;)
+		{
+			FChunkState NextWriteChunkState = NextWriteChunkStatePtr.Get();
+			if (NextWriteChunkState.GetReaderFlags() != 0)
+			{
+				if (!WriterEvent.Wait(TimeoutMs))
+				{
+					return nullptr;
+				}
+			}
+			else if (NextWriteChunkStatePtr.TryUpdate(NextWriteChunkState, FChunkState(EWriteState::Writing, WriterState.GetReaderFlags(), 0)))
+			{
+				FWriterStatePtr WriterStatePtr = Header->GetWriterStatePtr();
+				if (WriterStatePtr.TryUpdate(WriterState, FWriterState(NextWriteChunkIdx, WriterState.GetReaderFlags(), WriterState.GetRefCount(), NextWriteChunkIdx == 0)))
+				{
+					break;
+				}
+				else
+				{
+					WriterState = WriterStatePtr.Get();
+				}
+			}
+		}
+		return ChunkPtrs[NextWriteChunkIdx];
+	}
+
+private:
+	FComputeMemoryMappedFile MemoryMappedFile;
+
+	static void InitShared(FComputeBufferDetail* Detail)
+	{
+		FHeader* Header = Detail->Header;
+
+		unsigned char* NextPtr = (unsigned char*)(Header + 1);
+		for (int ChunkIdx = 0; ChunkIdx < Header->NumChunks; ChunkIdx++)
+		{
+			Detail->ChunkPtrs[ChunkIdx] = NextPtr;
+			NextPtr += Header->ChunkLength;
+		}
+	}
+
+	void SetAllReadEvents()
+	{
+		for (int Idx = 0; Idx < Header->NumReaders; Idx++)
+		{
+			ReaderEvents[Idx].Signal();
+		}
+	}
+};
+
+
+
+
+
+
+
+
+//// FComputeBuffer /////
+
+FComputeBuffer::FComputeBuffer()
+	: Detail(nullptr)
+{
+}
+
+FComputeBuffer::FComputeBuffer(const FComputeBuffer& Other)
+	: Detail(Other.Detail)
+{
+	if (Detail)
+	{
+		Detail->AddRef();
+	}
+}
+
+FComputeBuffer::FComputeBuffer(FComputeBuffer&& Other) noexcept
+	: Detail(Other.Detail)
+{
+	Other.Detail = nullptr;
+}
+
+FComputeBuffer::~FComputeBuffer()
+{
+	Close();
+}
+
+FComputeBuffer& FComputeBuffer::operator=(const FComputeBuffer& Other)
+{
+	Close();
+
+	if (Other.Detail)
+	{
+		Detail = Other.Detail;
+		Detail->AddRef();
+	}
+
+	return *this;
+}
+
+FComputeBuffer& FComputeBuffer::operator=(FComputeBuffer&& Other) noexcept
+{
+	Close();
+
+	Detail = Other.Detail;
+	Other.Detail = nullptr;
+
+	return *this;
+}
+
+bool FComputeBuffer::CreateNew(const FParams& Params)
+{
+	Close();
+
+	std::unique_ptr<FComputeBufferDetail> Buffer = FComputeBufferDetail::CreateNew(Params);
+	if (Buffer)
+	{
+		Detail = Buffer.release();
+		return true;
+	}
+	return false;
+}
+
+bool FComputeBuffer::OpenExisting(const char* Name)
+{
+	Close();
+
+	std::unique_ptr<FComputeBufferDetail> Buffer = FComputeBufferDetail::OpenExisting(Name);
+	if (Buffer)
+	{
+		Detail = Buffer.release();
+		return true;
+	}
+	return false;
+}
+
+void FComputeBuffer::Close()
+{
+	if (Detail != nullptr)
+	{
+		Detail->Release();
+		Detail = nullptr;
+	}
+}
+
+FComputeBufferReader FComputeBuffer::CreateReader()
+{
+	int ReaderIdx = Detail->CreateReader();
+	Detail->AddRef();
+	return FComputeBufferReader(Detail, ReaderIdx);
+}
+
+FComputeBufferWriter FComputeBuffer::CreateWriter()
+{
+	Detail->CreateWriter();
+	Detail->AddRef();
+	return FComputeBufferWriter(Detail);
+}
+
+const char* FComputeBuffer::GetName() const
+{
+	return Detail->Name;
+}
+
+
+
+
+
+
+//// FComputeBufferReader /////
+
+FComputeBufferReader::FComputeBufferReader()
+	: Detail(nullptr)
+	, ReaderIdx(0)
+{
+}
+
+FComputeBufferReader::FComputeBufferReader(const FComputeBufferReader& Other)
+	: Detail(Other.Detail)
+	, ReaderIdx(Other.ReaderIdx)
+{
+	if (Detail)
+	{
+		Detail->AddRef();
+		Detail->AddReaderRef(ReaderIdx);
+	}
+}
+
+FComputeBufferReader::FComputeBufferReader(FComputeBufferReader&& Other) noexcept
+	: Detail(Other.Detail)
+	, ReaderIdx(Other.ReaderIdx)
+{
+	Other.Detail = nullptr;
+}
+
+FComputeBufferReader::~FComputeBufferReader()
+{
+	Close();
+}
+
+FComputeBufferReader& FComputeBufferReader::operator=(const FComputeBufferReader& Other)
+{
+	Close();
+
+	if (Other.Detail)
+	{
+		Detail = Other.Detail;
+		ReaderIdx = Other.ReaderIdx;
+
+		Detail->AddRef();
+		Detail->AddReaderRef(ReaderIdx);
+	}
+
+	return *this;
+}
+
+FComputeBufferReader& FComputeBufferReader::operator=(FComputeBufferReader&& Other) noexcept
+{
+	Close();
+
+	Detail = Other.Detail;
+	ReaderIdx = Other.ReaderIdx;
+
+	Other.Detail = nullptr;
+	return *this;
+}
+
+void FComputeBufferReader::Close()
+{
+	if (Detail != nullptr)
+	{
+		Detail->ReleaseReaderRef(ReaderIdx);
+		Detail->Release();
+		Detail = nullptr;
+	}
+}
+
+void FComputeBufferReader::Detach()
+{
+	UE_COMPUTE_ASSERT(Detail);
+	Detail->DetachReader(ReaderIdx);
+}
+
+bool FComputeBufferReader::IsComplete() const
+{
+	UE_COMPUTE_ASSERT(Detail);
+	return Detail->IsComplete(ReaderIdx);
+}
+
+void FComputeBufferReader::AdvanceReadPosition(size_t Size)
+{
+	UE_COMPUTE_ASSERT(Detail);
+	Detail->AdvanceReadPosition(ReaderIdx, Size);
+}
+
+size_t FComputeBufferReader::GetMaxReadSize() const
+{
+	UE_COMPUTE_ASSERT(Detail);
+	return Detail->GetMaxReadSize(ReaderIdx);
+}
+
+size_t FComputeBufferReader::Read(void* Buffer, size_t MaxSize, int TimeoutMs)
+{
+	const unsigned char* RecvData = WaitToRead(1, TimeoutMs);
+	if (RecvData == nullptr)
+	{
+		return 0;
+	}
+
+	size_t RecvSize = GetMaxReadSize();
+	if (MaxSize < RecvSize)
+	{
+		RecvSize = MaxSize;
+	}
+
+	memcpy(Buffer, RecvData, RecvSize);
+	AdvanceReadPosition(RecvSize);
+	return RecvSize;
+}
+
+const unsigned char* FComputeBufferReader::WaitToRead(size_t MinSize, int TimeoutMs)
+{
+	UE_COMPUTE_ASSERT(Detail);
+	return Detail->WaitToRead(ReaderIdx, MinSize, TimeoutMs);
+}
+
+FComputeBufferReader::FComputeBufferReader(FComputeBufferDetail* InDetail, int InReaderIdx)
+	: Detail(InDetail)
+	, ReaderIdx(InReaderIdx)
+{
+}
+
+const char* FComputeBufferReader::GetName() const
+{
+	UE_COMPUTE_ASSERT(Detail);
+	return Detail->Name;
+}
+
+
+
+
+
+//// FComputeBufferWriter /////
+
+FComputeBufferWriter::FComputeBufferWriter()
+	: Detail(nullptr)
+{
+}
+
+FComputeBufferWriter::FComputeBufferWriter(const FComputeBufferWriter& Other)
+	: Detail(Other.Detail)
+{
+	if (Detail)
+	{
+		Detail->AddRef();
+		Detail->AddWriterRef();
+	}
+}
+
+FComputeBufferWriter::FComputeBufferWriter(FComputeBufferWriter&& Other) noexcept
+	: Detail(Other.Detail)
+{
+	Other.Detail = nullptr;
+}
+
+FComputeBufferWriter::~FComputeBufferWriter()
+{
+	Close();
+}
+
+FComputeBufferWriter& FComputeBufferWriter::operator=(const FComputeBufferWriter& Other)
+{
+	Close();
+
+	if (Other.Detail)
+	{
+		Detail = Other.Detail;
+		Detail->AddWriterRef();
+		Detail->AddRef();
+	}
+
+	return *this;
+}
+
+FComputeBufferWriter& FComputeBufferWriter::operator=(FComputeBufferWriter&& Other) noexcept
+{
+	Close();
+
+	Detail = Other.Detail;
+	Other.Detail = nullptr;
+
+	return *this;
+}
+
+
+void FComputeBufferWriter::Close()
+{
+	if (Detail != nullptr)
+	{
+		Detail->ReleaseWriterRef();
+		Detail->Release();
+		Detail = nullptr;
+	}
+}
+
+void FComputeBufferWriter::MarkComplete()
+{
+	Detail->MarkComplete();
+}
+
+void FComputeBufferWriter::AdvanceWritePosition(size_t Size)
+{
+	Detail->AdvanceWritePosition(Size);
+}
+
+size_t FComputeBufferWriter::GetMaxWriteSize() const
+{
+	return Detail->GetMaxWriteSize();
+}
+
+size_t FComputeBufferWriter::Write(const void* Buffer, size_t MaxSize, int TimeoutMs)
+{
+	unsigned char* SendData = WaitToWrite(1, TimeoutMs);
+	if (SendData == nullptr)
+	{
+		return 0;
+	}
+
+	size_t SendSize = GetMaxWriteSize();
+	if (MaxSize < SendSize)
+	{
+		SendSize = MaxSize;
+	}
+
+	memcpy(SendData, Buffer, SendSize);
+	AdvanceWritePosition(SendSize);
+	return SendSize;
+}
+
+unsigned char* FComputeBufferWriter::WaitToWrite(size_t MinSize, int TimeoutMs)
+{
+	return Detail->WaitToWrite(MinSize, TimeoutMs);
+}
+
+FComputeBufferWriter::FComputeBufferWriter(FComputeBufferDetail* InDetail)
+	: Detail(InDetail)
+{
+}
+
+const char* FComputeBufferWriter::GetName() const
+{
+	return Detail->Name;
+}
diff --git a/Engine/Source/Developer/Horde/Private/ComputeChannel.cpp b/Engine/Source/Developer/Horde/Private/ComputeChannel.cpp
new file mode 100644
--- /dev/null
+++ b/Engine/Source/Developer/Horde/Private/ComputeChannel.cpp
@@ -0,0 +1,40 @@
+// Copyright Epic Games, Inc. All Rights Reserved.
+
+#include "ComputeChannel.h"
+#include <string.h>
+#include <stdio.h>
+#include "ComputeSocket.h"
+
+FComputeChannel::FComputeChannel()
+{
+}
+
+FComputeChannel::FComputeChannel(FComputeBufferReader InReader, FComputeBufferWriter InWriter)
+	: Reader(std::move(InReader))
+	, Writer(std::move(InWriter))
+{
+}
+
+FComputeChannel::~FComputeChannel()
+{
+}
+
+bool FComputeChannel::IsValid() const
+{
+	return Reader.IsValid();
+}
+
+size_t FComputeChannel::Send(const void* Data, size_t Size, int TimeoutMs)
+{
+	return Writer.Write(Data, Size, TimeoutMs);
+}
+
+size_t FComputeChannel::Recv(void* Data, size_t Size, int TimeoutMs)
+{
+	return Reader.Read(Data, Size, TimeoutMs);
+}
+
+void FComputeChannel::MarkComplete()
+{
+	Writer.MarkComplete();
+}
diff --git a/Engine/Source/Developer/Horde/Private/ComputePlatform.cpp b/Engine/Source/Developer/Horde/Private/ComputePlatform.cpp
new file mode 100644
--- /dev/null
+++ b/Engine/Source/Developer/Horde/Private/ComputePlatform.cpp
@@ -0,0 +1,497 @@
+// Copyright Epic Games, Inc. All Rights Reserved.
+
+#include "ComputePlatform.h"
+#include <assert.h>
+#include <wchar.h>
+#include <bit>
+#include <algorithm>
+#include <iostream>
+
+#if UE_COMPUTE_PLATFORM_WINDOWS
+	#include <Windows.h>
+	#undef min
+	#undef max
+	#undef GetEnvironmentVariable
+	#undef SendMessage
+#else
+	#include <semaphore.h>
+	#include <unistd.h>
+	#include <atomic>
+	#include <sys/mman.h>
+	#include <sys/stat.h>
+	#include <time.h>
+	#include <fcntl.h>
+
+	struct FFormatErrno { };
+
+	static std::ostream &operator<<(std::ostream &os, const FFormatErrno &)
+	{
+		os << "(errno=" << errno << ": " << strerror(errno) << ")";
+		return os;
+	}
+#endif
+
+///////////////////////////////////////////////////
+
+FComputeEvent::FComputeEvent()
+	: Handle(nullptr)
+{
+}
+
+FComputeEvent::~FComputeEvent()
+{
+	Close();
+}
+
+bool FComputeEvent::Create(const char* Name)
+{
+	Close();
+
+#if UE_COMPUTE_PLATFORM_WINDOWS
+	Handle = CreateEventA(NULL, FALSE, FALSE, Name);
+	return Handle != nullptr;
+#else
+	sem_t* Value = sem_open(Name, O_CREAT | O_EXCL, 0666, 1);
+	if(Value != SEM_FAILED)
+	{
+		Handle = Value;
+		return true;
+	}
+	return false;
+#endif
+}
+
+bool FComputeEvent::OpenExisting(const char* Name)
+{
+	Close();
+
+#if UE_COMPUTE_PLATFORM_WINDOWS
+	Handle = OpenEventA(SYNCHRONIZE | EVENT_MODIFY_STATE, TRUE, Name);
+	return Handle != nullptr;
+#else
+	sem_t* Value = sem_open(Name, 0);
+	if(Value != SEM_FAILED)
+	{
+		Handle = Value;
+		return true;
+	}
+	return false;
+#endif
+}
+
+void FComputeEvent::Close()
+{
+#if UE_COMPUTE_PLATFORM_WINDOWS
+	if (Handle != nullptr)
+	{
+		CloseHandle(Handle);
+		Handle = nullptr;
+	}
+#else
+	if (Handle != nullptr)
+	{
+		sem_close((sem_t*)Handle);
+		Handle = nullptr;
+	}
+#endif
+}
+
+void FComputeEvent::Signal()
+{
+#if UE_COMPUTE_PLATFORM_WINDOWS
+	SetEvent(Handle);
+#else
+	sem_post((sem_t*)Handle);
+#endif
+}
+
+bool FComputeEvent::Wait(int timeoutMs)
+{
+#if UE_COMPUTE_PLATFORM_WINDOWS
+	DWORD WaitParam = (timeoutMs < 0) ? INFINITE : (DWORD)timeoutMs;
+	return WaitForSingleObject(Handle, WaitParam) != WAIT_TIMEOUT;
+#else
+	if (timeoutMs == -1)
+	{
+		return sem_wait((sem_t*)Handle) == 0;
+	}
+
+	if (sem_trywait((sem_t*)Handle) == 0)
+	{
+		return true;
+	}
+
+	if(timeoutMs == 0)
+	{
+		return false;
+	}
+		
+	struct timespec ts;
+	if (clock_gettime(CLOCK_REALTIME, &ts) == -1)
+	{
+		UE_COMPUTE_ASSERT(false);
+		return false;
+	}
+	
+	const long long NsPerSec = 1'000'000'000;
+	const long long NsPerMs = NsPerSec / 1000;
+	
+	long long newNs = (long long)ts.tv_nsec + (timeoutMs * NsPerMs);
+	ts.tv_nsec = newNs % NsPerSec;
+	ts.tv_sec += newNs / NsPerSec;
+	
+#if UE_COMPUTE_PLATFORM_MAC
+	for(;;)
+	{
+		if (sem_trywait((sem_t*)Handle) == 0)
+		{
+			return true;
+		}
+		
+		struct timespec currentTs;
+		if (clock_gettime(CLOCK_REALTIME, &currentTs) != 0 || currentTs.tv_sec > ts.tv_sec || (currentTs.tv_sec == ts.tv_sec && currentTs.tv_nsec > ts.tv_nsec))
+		{
+			return false;
+		}
+		
+		struct timespec sleepTs = { 0, };
+		sleepTs.tv_nsec = 100 * NsPerMs;
+		
+		nanosleep(&sleepTs, nullptr);
+	}
+#else
+	return sem_timedwait((sem_t*)Handle, &ts) == 0;
+#endif
+#endif
+}
+
+/////////////////////////////////////////////////// 
+
+FComputeMemoryMappedFile::FComputeMemoryMappedFile()
+	: Handle(nullptr)
+	, Pointer(nullptr)
+	, MappedSize(0)
+	, OwnerName(nullptr)
+{
+}
+
+FComputeMemoryMappedFile::~FComputeMemoryMappedFile()
+{
+	Close();
+}
+
+bool FComputeMemoryMappedFile::Create(const char* Name, long long Capacity)
+{
+	Close();
+
+#if UE_COMPUTE_PLATFORM_WINDOWS
+	LARGE_INTEGER LargeInteger;
+	LargeInteger.QuadPart = Capacity;
+
+	Handle = CreateFileMappingA(INVALID_HANDLE_VALUE, NULL, PAGE_READWRITE, LargeInteger.HighPart, LargeInteger.LowPart, Name);
+	if (Handle == nullptr)
+	{
+		return false;
+	}
+
+	Pointer = MapViewOfFile(Handle, FILE_MAP_ALL_ACCESS, 0, 0, 0);
+	if (Pointer == nullptr)
+	{
+		return false;
+	}
+
+	return true;
+#else
+	int Fd = shm_open(Name, O_CREAT | O_EXCL | O_RDWR, 0666);
+	if(Fd < 0)
+	{
+		std::cerr << "Unable to create shared memory object '" << Name << "' " << FFormatErrno() << std::endl;
+		return false;
+	}
+
+	Handle = (void*)(size_t)Fd;
+	MappedSize = Capacity;
+	OwnerName = strdup(Name);
+
+	if(ftruncate(Fd, MappedSize) < 0)
+	{
+		std::cerr << "Unable to update size of shared memory object '" << Name << "' to " << MappedSize << " " << FFormatErrno() << std::endl;
+		return false;
+	}
+
+	Pointer = mmap(nullptr, MappedSize, PROT_READ | PROT_WRITE, MAP_SHARED, Fd, 0);
+	if(Pointer == MAP_FAILED)
+	{
+		std::cerr << "Unable to map shared memory object '" << Name << " " << FFormatErrno() << std::endl;
+		return false;
+	}
+	
+	return true;
+#endif
+}
+
+bool FComputeMemoryMappedFile::OpenExisting(const char* Name)
+{
+	Close();
+
+#if UE_COMPUTE_PLATFORM_WINDOWS
+	Handle = OpenFileMappingA(FILE_MAP_ALL_ACCESS, TRUE, Name);
+	if (Handle == nullptr)
+	{
+		return false;
+	}
+
+	Pointer = MapViewOfFile(Handle, FILE_MAP_ALL_ACCESS, 0, 0, 0);
+	if (Pointer == nullptr)
+	{
+		return false;
+	}
+
+	return true;
+#else
+	int Fd = shm_open(Name, O_RDWR, 0666);
+	if(Fd < 0)
+	{
+		std::cerr << "Unable to open shared memory object '" << Name << "' " << FFormatErrno() << std::endl;
+		return false;
+	}
+
+	Handle = (void*)(size_t)Fd;
+
+	struct stat st;
+	fstat(Fd, &st);
+	MappedSize = st.st_size;
+
+	Pointer = mmap(nullptr, MappedSize, PROT_READ | PROT_WRITE, MAP_SHARED, Fd, 0);
+	if(Pointer == MAP_FAILED)
+	{
+		std::cerr << "Unable to map shared memory object '" << Name << " " << FFormatErrno() << std::endl;
+		return false;
+	}
+	
+	return true;
+#endif
+}
+
+void FComputeMemoryMappedFile::Close()
+{
+#if UE_COMPUTE_PLATFORM_WINDOWS
+	if (Pointer != nullptr)
+	{
+		UnmapViewOfFile(Pointer);
+		Pointer = nullptr;
+	}
+
+	if (Handle != nullptr)
+	{
+		CloseHandle(Handle);
+		Handle = nullptr;
+	}
+#else
+	if(OwnerName != nullptr)
+	{
+		shm_unlink(OwnerName);
+		OwnerName = nullptr;
+	}
+	
+	if(Pointer != nullptr)
+	{
+		munmap(Pointer, MappedSize);
+		Pointer = nullptr;
+	}
+
+	int Fd = (int)(size_t)Handle;
+	if(Fd >= 0)
+	{
+		close(Fd);
+		Handle = nullptr;
+	}
+#endif
+}
+
+void* FComputeMemoryMappedFile::GetPointer() const
+{
+#if UE_COMPUTE_PLATFORM_WINDOWS
+	return Pointer;
+#else
+	return (unsigned char*)Pointer + 16;
+#endif
+}
+
+/////////////////////////////////////////////////// 
+
+void FComputePlatform::AssertFailed(const char* Expr, const char* File, int Line)
+{
+	std::cerr << "Assertion failed: '" << Expr << std::endl;
+	std::cerr << "  at " << File << "(" << Line << ")" << std::endl;
+#if UE_COMPUTE_PLATFORM_WINDOWS
+	DebugBreak();
+#endif
+	exit(1);
+}
+
+bool FComputePlatform::GetEnvironmentVariable(const char* Name, char* Buffer, size_t BufferLen)
+{
+#if UE_COMPUTE_PLATFORM_WINDOWS
+	int Length = GetEnvironmentVariableA(Name, Buffer, (DWORD)BufferLen);
+	return Length > 0 && Length < BufferLen;
+#else
+	char* Value = getenv(Name);
+	if(Value != nullptr)
+	{
+		FComputePlatform::Strcpy(Buffer, BufferLen, Value);
+		return true;
+	}
+	return false;
+#endif
+}
+
+void FComputePlatform::CreateUniqueName(char* NameBuffer, size_t NameBufferLen)
+{
+	static long Counter = 0;
+
+#if UE_COMPUTE_PLATFORM_WINDOWS
+	DWORD Pid = GetCurrentProcessId();
+	ULONGLONG TickCount = GetTickCount64();
+	snprintf(NameBuffer, NameBufferLen, "Local\\COMPUTE_%lu_%llu_%lu", Pid, TickCount, AtomicIncrement(&Counter));
+#else
+	struct timespec ts;
+	clock_gettime(CLOCK_MONOTONIC, &ts);
+	snprintf(NameBuffer, NameBufferLen, "/UEC_%u%zu%zu_%lu", getpid(), (size_t)ts.tv_sec, (size_t)ts.tv_nsec, AtomicIncrement(&Counter));
+#endif
+}
+
+unsigned int FComputePlatform::FloorLog2(unsigned int Value)
+{
+	return std::max<unsigned int>(0, 31 - std::countl_zero(Value));
+}
+
+unsigned int FComputePlatform::CountLeadingZeros(unsigned int Value)
+{
+	return std::countl_zero(Value);
+}
+
+void FComputePlatform::Strcpy(char* Dest, size_t DestLen, const char* Source)
+{
+	size_t Length = std::min(strlen(Source), DestLen - 1);
+	memcpy(Dest, Source, Length);
+	Dest[Length] = 0;
+}
+
+int FComputePlatform::Stricmp(const char* A, const char* B)
+{
+#if UE_COMPUTE_PLATFORM_WINDOWS
+	return _stricmp(A, B);
+#else
+	return strcasecmp(A, B);
+#endif
+}
+
+long long FComputePlatform::AtomicRead64(const volatile long long* Ptr)
+{
+#if UE_COMPUTE_PLATFORM_WINDOWS
+	return InterlockedCompareExchange64(const_cast<volatile long long*>(Ptr), 0, 0);
+#else
+	return std::atomic_load((std::atomic<long long>*)Ptr);
+#endif
+}
+
+void FComputePlatform::AtomicWrite64(volatile long long* Ptr, long long Value)
+{
+#if UE_COMPUTE_PLATFORM_WINDOWS
+	InterlockedExchange64(Ptr, Value);
+#else
+	return std::atomic_store((std::atomic<long long>*)Ptr, Value);
+#endif
+}
+
+long FComputePlatform::AtomicIncrement(volatile long* Ptr)
+{
+#if UE_COMPUTE_PLATFORM_WINDOWS
+	return InterlockedIncrement(Ptr);
+#else
+	return std::atomic_fetch_add((std::atomic<long>*)Ptr, 1) + 1;
+#endif
+}
+
+long FComputePlatform::AtomicDecrement(volatile long* Ptr)
+{
+#if UE_COMPUTE_PLATFORM_WINDOWS
+	return InterlockedDecrement(Ptr);
+#else
+	return std::atomic_fetch_add((std::atomic<long>*)Ptr, -1) - 1;
+#endif
+}
+
+long long FComputePlatform::AtomicAdd64(volatile long long* Ptr, long long Value)
+{
+#if UE_COMPUTE_PLATFORM_WINDOWS
+	return InterlockedAdd64(Ptr, Value);
+#else
+	return std::atomic_fetch_add((std::atomic<long long>*)Ptr, Value) + Value;
+#endif
+}
+
+long FComputePlatform::AtomicAnd(volatile long* Ptr, long Value)
+{
+#if UE_COMPUTE_PLATFORM_WINDOWS
+	return InterlockedAnd(Ptr, Value);
+#else
+	return std::atomic_fetch_and((std::atomic<long>*)Ptr, Value);
+#endif
+}
+
+long long FComputePlatform::AtomicAnd64(volatile long long* Ptr, long long Value)
+{
+#if UE_COMPUTE_PLATFORM_WINDOWS
+	return InterlockedAnd64(Ptr, Value);
+#else
+	return std::atomic_fetch_and((std::atomic<long long>*)Ptr, Value);
+#endif
+}
+
+long long FComputePlatform::AtomicOr64(volatile long long* Ptr, long long Value)
+{
+#if UE_COMPUTE_PLATFORM_WINDOWS
+	return InterlockedOr64(Ptr, Value);
+#else
+	return std::atomic_fetch_or((std::atomic<long long>*)Ptr, Value);
+#endif
+}
+
+long long FComputePlatform::AtomicIncrement64(volatile long long* Ptr)
+{
+#if UE_COMPUTE_PLATFORM_WINDOWS
+	return InterlockedIncrement64(Ptr);
+#else
+	return std::atomic_fetch_add((std::atomic<long long>*)Ptr, 1) + 1;
+#endif
+}
+
+long long FComputePlatform::AtomicExchange64(volatile long long* Ptr, long long Exchange)
+{
+#if UE_COMPUTE_PLATFORM_WINDOWS
+	return InterlockedExchange64(Ptr, Exchange);
+#else
+	return std::atomic_exchange((std::atomic<long long>*)Ptr, Exchange);
+#endif
+}
+
+bool FComputePlatform::AtomicCompareExchange(volatile long* Ptr, long Exchange, long Comperand)
+{
+#if UE_COMPUTE_PLATFORM_WINDOWS
+	return InterlockedCompareExchange(Ptr, Exchange, Comperand) == Comperand;
+#else
+	return std::atomic_compare_exchange_strong((std::atomic<long>*)Ptr, &Comperand, Exchange);
+#endif
+}
+
+bool FComputePlatform::AtomicCompareExchange64(volatile long long* Ptr, long long Exchange, long long Comperand)
+{
+#if UE_COMPUTE_PLATFORM_WINDOWS
+	return InterlockedCompareExchange64(Ptr, Exchange, Comperand) == Comperand;
+#else
+	return std::atomic_compare_exchange_strong((std::atomic<long long>*)Ptr, &Comperand, Exchange);
+#endif
+}
+
diff --git a/Engine/Source/Developer/Horde/Private/ComputeSocket.cpp b/Engine/Source/Developer/Horde/Private/ComputeSocket.cpp
new file mode 100644
--- /dev/null
+++ b/Engine/Source/Developer/Horde/Private/ComputeSocket.cpp
@@ -0,0 +1,399 @@
+// Copyright Epic Games, Inc. All Rights Reserved.
+
+#include "ComputeSocket.h"
+#include "ComputePlatform.h"
+#include <iostream>
+#include <assert.h>
+#include <unordered_set>
+#include <unordered_map>
+#include <vector>
+#include <thread>
+#include <mutex>
+
+FComputeSocket::FComputeSocket()
+{
+}
+
+FComputeSocket::~FComputeSocket()
+{
+}
+
+FComputeChannel FComputeSocket::CreateChannel(int ChannelId)
+{
+	FComputeBuffer RecvBuffer;
+	if (!RecvBuffer.CreateNew(FComputeBuffer::FParams()))
+	{
+		return FComputeChannel();
+	}
+
+	FComputeBuffer SendBuffer;
+	if (!SendBuffer.CreateNew(FComputeBuffer::FParams()))
+	{
+		return FComputeChannel();
+	}
+
+	return CreateChannel(ChannelId, std::move(RecvBuffer), std::move(SendBuffer));
+}
+
+FComputeChannel FComputeSocket::CreateChannel(int ChannelId, FComputeBuffer RecvBuffer, FComputeBuffer SendBuffer)
+{
+	FComputeChannel Channel(RecvBuffer.CreateReader(), SendBuffer.CreateWriter());
+
+	AttachRecvBuffer(ChannelId, std::move(RecvBuffer));
+	AttachSendBuffer(ChannelId, std::move(SendBuffer));
+
+	return Channel;
+}
+
+//////////////////////////////////////////////////////
+
+const char* const FWorkerComputeSocket::IpcEnvVar = "UE_HORDE_COMPUTE_IPC";
+
+enum class FWorkerComputeSocket::EMessageType
+{
+	AttachRecvBuffer = 0,
+	AttachSendBuffer = 1,
+};
+
+FWorkerComputeSocket::FWorkerComputeSocket()
+{
+}
+
+FWorkerComputeSocket::~FWorkerComputeSocket()
+{
+	Close();
+}
+
+bool FWorkerComputeSocket::Open()
+{
+	char EnvVar[FComputeBuffer::MaxNameLength];
+	if (!FComputePlatform::GetEnvironmentVariable(IpcEnvVar, EnvVar, sizeof(EnvVar) / sizeof(EnvVar[0])))
+	{
+		return false;
+	}
+
+	return Open(EnvVar);
+}
+
+bool FWorkerComputeSocket::Open(const char* CommandBufferName)
+{
+	FComputeBuffer CommandBuffer;
+	if (CommandBuffer.OpenExisting(CommandBufferName))
+	{
+		CommandBufferWriter = CommandBuffer.CreateWriter();
+		return true;
+	}
+	return false;
+}
+
+void FWorkerComputeSocket::Close()
+{
+	CommandBufferWriter.Close();
+}
+
+void FWorkerComputeSocket::AttachRecvBuffer(int ChannelId, FComputeBuffer RecvBuffer)
+{
+	AttachBuffer(ChannelId, EMessageType::AttachRecvBuffer, RecvBuffer.GetName());
+	Buffers.push_back(std::move(RecvBuffer));
+}
+
+void FWorkerComputeSocket::AttachSendBuffer(int ChannelId, FComputeBuffer SendBuffer)
+{
+	AttachBuffer(ChannelId, EMessageType::AttachSendBuffer, SendBuffer.GetName());
+	Buffers.push_back(std::move(SendBuffer));
+}
+
+void FWorkerComputeSocket::AttachBuffer(int ChannelId, EMessageType Type, const char* Name)
+{
+	unsigned char* Data = CommandBufferWriter.WaitToWrite(1024);
+
+	size_t Len = 0;
+	Len += WriteVarUInt(Data + Len, (unsigned char)Type);
+	Len += WriteVarUInt(Data + Len, (unsigned int)ChannelId);
+	Len += WriteString(Data + Len, Name);
+
+	CommandBufferWriter.AdvanceWritePosition(Len);
+}
+
+void FWorkerComputeSocket::RunServer(FComputeBufferReader& CommandBufferReader, FComputeSocket& Socket)
+{
+	const unsigned char* Message;
+	while ((Message = CommandBufferReader.WaitToRead(1)) != nullptr)
+	{
+		size_t Len = 0;
+
+		unsigned int Type;
+		Len += ReadVarUInt(Message + Len, &Type);
+
+		EMessageType MessageType = (EMessageType)*Message;
+		switch (MessageType)
+		{
+		case EMessageType::AttachSendBuffer:
+			{
+				unsigned int ChannelId;
+				Len += ReadVarUInt(Message + Len, &ChannelId);
+
+				char Name[FComputeBuffer::MaxNameLength];
+				Len += ReadString(Message + Len, Name, FComputeBuffer::MaxNameLength);
+
+				FComputeBuffer Buffer;
+				if (Buffer.OpenExisting(Name))
+				{
+					Socket.AttachSendBuffer(ChannelId, Buffer);
+				}
+				else
+				{
+					UE_COMPUTE_ASSERT(false);
+				}
+			}
+			break;
+		case EMessageType::AttachRecvBuffer:
+			{
+				unsigned int ChannelId;
+				Len += ReadVarUInt(Message + Len, &ChannelId);
+
+				char Name[FComputeBuffer::MaxNameLength];
+				Len += ReadString(Message + Len, Name, FComputeBuffer::MaxNameLength);
+
+				FComputeBuffer Buffer;
+				if (Buffer.OpenExisting(Name))
+				{
+					Socket.AttachRecvBuffer(ChannelId, Buffer);
+				}
+				else
+				{
+					UE_COMPUTE_ASSERT(false);
+				}
+			}
+			break;
+		default:
+			UE_COMPUTE_ASSERT(false);
+			return;
+		}
+
+		CommandBufferReader.AdvanceReadPosition(Len);
+	}
+}
+
+size_t FWorkerComputeSocket::ReadVarUInt(const unsigned char* Pos, unsigned int* OutValue)
+{
+	size_t ByteCount = FComputePlatform::CountLeadingZeros((unsigned char)(~*static_cast<const unsigned char*>(Pos))) - 23;
+
+	unsigned int Value = *Pos++ & (unsigned char)(0xff >> ByteCount);
+	switch (ByteCount - 1)
+	{
+	case 8: Value <<= 8; Value |= *Pos++;
+	case 7: Value <<= 8; Value |= *Pos++;
+	case 6: Value <<= 8; Value |= *Pos++;
+	case 5: Value <<= 8; Value |= *Pos++;
+	case 4: Value <<= 8; Value |= *Pos++;
+	case 3: Value <<= 8; Value |= *Pos++;
+	case 2: Value <<= 8; Value |= *Pos++;
+	case 1: Value <<= 8; Value |= *Pos++;
+	default:
+		break;
+	}
+
+	*OutValue = Value;
+	return ByteCount;
+}
+
+size_t FWorkerComputeSocket::ReadString(const unsigned char* Pos, char* OutText, size_t OutTextMaxLen)
+{
+	unsigned int TextLen;
+
+	size_t Len = ReadVarUInt(Pos, &TextLen);
+	FComputePlatform::Strcpy(OutText, OutTextMaxLen, (const char*)Pos + Len);
+
+	return Len + TextLen;
+}
+
+size_t FWorkerComputeSocket::WriteVarUInt(unsigned char* Pos, unsigned int Value)
+{
+	// Use BSR to return the log2 of the integer
+	// return 0 if value is 0
+	unsigned int ByteCount = (unsigned int)(int(FComputePlatform::FloorLog2(Value)) / 7 + 1);
+
+	unsigned char* OutBytes = Pos + ByteCount - 1;
+	switch (ByteCount - 1)
+	{
+	case 4: *OutBytes-- = (unsigned char)(Value); Value >>= 8; [[fallthrough]];
+	case 3: *OutBytes-- = (unsigned char)(Value); Value >>= 8; [[fallthrough]];
+	case 2: *OutBytes-- = (unsigned char)(Value); Value >>= 8; [[fallthrough]];
+	case 1: *OutBytes-- = (unsigned char)(Value); Value >>= 8; [[fallthrough]];
+	default: 
+		break;
+	}
+	*OutBytes = (unsigned char)(0xff << (9 - ByteCount)) | (unsigned char)(Value);
+
+	return ByteCount;
+}
+
+size_t FWorkerComputeSocket::WriteString(unsigned char* Pos, const char* Text)
+{
+	size_t TextLen = strlen(Text);
+
+	size_t Len = WriteVarUInt(Pos, (int)TextLen);
+	memcpy((char*)Pos + Len, Text, TextLen);
+
+	return Len + TextLen;
+}
+
+//////////////////////////////////////////////////////
+
+class FRemoteComputeSocket : public FComputeSocket
+{
+public:
+	enum class EControlMessageType
+	{
+		Detach = -2,
+	};
+
+	struct FFrameHeader
+	{
+		int Channel;
+		int Size;
+	};
+
+	std::unique_ptr<FComputeTransport> Transport;
+	const EComputeSocketEndpoint Endpoint;
+	std::mutex CriticalSection;
+
+	std::thread RecvThread;
+
+	std::unordered_map<int, FComputeBufferWriter> Writers;
+	std::vector<FComputeBufferReader> Readers;
+	std::unordered_map<int, std::thread> SendThreads;
+
+	FRemoteComputeSocket(std::unique_ptr<FComputeTransport> InTransport, EComputeSocketEndpoint InEndpoint)
+		: Transport(std::move(InTransport))
+		, Endpoint(InEndpoint)
+		, CriticalSection()
+		, RecvThread(&FRemoteComputeSocket::RecvThreadProc, this)
+	{
+	}
+
+	~FRemoteComputeSocket()
+	{
+		for (FComputeBufferReader& Reader : Readers)
+		{
+			Reader.Detach();
+		}
+
+		for (std::pair<const int, std::thread>& Pair : SendThreads)
+		{
+			Pair.second.join();
+		}
+
+		Transport->Close();
+		RecvThread.join();
+	}
+
+	void RecvThreadProc()
+	{
+		std::unordered_map<int, FComputeBufferWriter> CachedWriters;
+
+		// Process messages from the remote
+		FFrameHeader Header;
+		while (Transport->RecvMessage(&Header, sizeof(Header)))
+		{
+			if (Header.Size >= 0)
+			{
+				ReadFrame(CachedWriters, Header.Channel, Header.Size);
+			}
+			else if (Header.Size == (int)EControlMessageType::Detach)
+			{
+				DetachRecvBuffer(CachedWriters, Header.Channel);
+			}
+			else
+			{
+				UE_COMPUTE_ASSERT(false);
+			}
+		}
+	}
+
+	void SendThreadProc(int Channel, FComputeBufferReader Reader)
+	{
+		FFrameHeader Header;
+		Header.Channel = Channel;
+
+		const unsigned char* Data;
+		while ((Data = Reader.WaitToRead(1)) != nullptr)
+		{
+			std::lock_guard<std::mutex> Lock(CriticalSection);
+			Header.Size = (int)Reader.GetMaxReadSize();
+			Transport->SendMessage(&Header, sizeof(Header));
+			Transport->SendMessage(Data, Header.Size);
+			Reader.AdvanceReadPosition(Header.Size);
+		}
+
+		if (Reader.IsComplete())
+		{
+			std::lock_guard<std::mutex> Lock(CriticalSection);
+			Header.Size = (int)EControlMessageType::Detach;
+			Transport->SendMessage(&Header, sizeof(Header));
+		}
+	}
+
+	bool ReadFrame(std::unordered_map<int, FComputeBufferWriter>& CachedWriters, int Channel, int Size)
+	{
+		std::unordered_map<int, FComputeBufferWriter>::iterator Iter = CachedWriters.find(Channel);
+		if (Iter == CachedWriters.end())
+		{
+			std::lock_guard<std::mutex> Lock(CriticalSection);
+
+			Iter = Writers.find(Channel);
+			if (Iter == Writers.end())
+			{
+				return false;
+			}
+
+			Iter = CachedWriters.insert(*Iter).first;
+		}
+
+		FComputeBufferWriter& Writer = Iter->second;
+
+		unsigned char* Data = Writer.WaitToWrite(Size);
+		if (!Transport->RecvMessage(Data, Size))
+		{
+			return false;
+		}
+
+		Writer.AdvanceWritePosition(Size);
+		return true;
+	}
+
+	void AttachRecvBuffer(int ChannelId, FComputeBuffer RecvBuffer) override
+	{
+		std::lock_guard<std::mutex> Lock(CriticalSection);
+		FComputeBufferWriter Writer = RecvBuffer.CreateWriter();
+		Writers.insert(std::pair<int, FComputeBufferWriter>(ChannelId, std::move(Writer)));
+	}
+
+	void AttachSendBuffer(int ChannelId, FComputeBuffer SendBuffer) override
+	{
+		std::lock_guard<std::mutex> Lock(CriticalSection);
+		FComputeBufferReader Reader = SendBuffer.CreateReader();
+		Readers.push_back(Reader);
+		SendThreads.insert(std::make_pair(ChannelId, std::thread(&FRemoteComputeSocket::SendThreadProc, this, ChannelId, std::move(Reader))));
+	}
+
+	void DetachRecvBuffer(std::unordered_map<int, FComputeBufferWriter>& CachedWriters, int Channel)
+	{
+		CachedWriters.erase(Channel);
+
+		std::lock_guard<std::mutex> Lock(CriticalSection);
+
+		std::unordered_map<int, FComputeBufferWriter>::iterator Iter = Writers.find(Channel);
+		if (Iter != Writers.end())
+		{
+			Iter->second.MarkComplete();
+			Writers.erase(Iter);
+		}
+	}
+};
+
+std::unique_ptr<FComputeSocket> CreateComputeSocket(std::unique_ptr<FComputeTransport> Transport, EComputeSocketEndpoint Endpoint)
+{
+	return std::unique_ptr<FComputeSocket>(new FRemoteComputeSocket(std::move(Transport), Endpoint));
+}
diff --git a/Engine/Source/Developer/Horde/Private/ComputeTransport.cpp b/Engine/Source/Developer/Horde/Private/ComputeTransport.cpp
new file mode 100644
--- /dev/null
+++ b/Engine/Source/Developer/Horde/Private/ComputeTransport.cpp
@@ -0,0 +1,82 @@
+// Copyright Epic Games, Inc. All Rights Reserved.
+
+#include "ComputeTransport.h"
+#include <algorithm>
+
+FComputeTransport::~FComputeTransport()
+{
+}
+
+bool FComputeTransport::SendMessage(const void* Data, size_t Size)
+{
+	const unsigned char* RemainingData = (const unsigned char*)Data;
+	for (size_t RemainingSize = Size; RemainingSize > 0; )
+	{
+		size_t SentSize = Send(RemainingData, RemainingSize);
+		if (SentSize == 0)
+		{
+			return false;
+		}
+
+		RemainingData += SentSize;
+		RemainingSize -= SentSize;
+	}
+	return true;
+}
+
+bool FComputeTransport::RecvMessage(void* Data, size_t Size)
+{
+	unsigned char* RemainingData = (unsigned char*)Data;
+	for (size_t RemainingSize = Size; RemainingSize > 0; )
+	{
+		size_t RecvSize = Recv(RemainingData, RemainingSize);
+		if (RecvSize == 0)
+		{
+			return false;
+		}
+
+		RemainingData += RecvSize;
+		RemainingSize -= RecvSize;
+	}
+	return true;
+}
+
+///////////////////////////////////
+
+FBufferTransport::FBufferTransport(FComputeBufferWriter InSendBufferWriter, FComputeBufferReader InRecvBufferReader)
+	: SendBufferWriter(std::move(InSendBufferWriter))
+	, RecvBufferReader(std::move(InRecvBufferReader))
+{
+}
+
+size_t FBufferTransport::Send(const void* Data, size_t Size)
+{
+	unsigned char* Buffer = SendBufferWriter.WaitToWrite(1);
+
+	size_t WriteSize = std::min(Size, SendBufferWriter.GetMaxWriteSize());
+	memcpy(Buffer, Data, WriteSize);
+	SendBufferWriter.AdvanceWritePosition(WriteSize);
+
+	return WriteSize;
+}
+
+size_t FBufferTransport::Recv(void* Data, size_t Size)
+{
+	const unsigned char* Buffer = RecvBufferReader.WaitToRead(1);
+
+	size_t ReadSize = std::min(Size, RecvBufferReader.GetMaxReadSize());
+	memcpy(Data, Buffer, ReadSize);
+	RecvBufferReader.AdvanceReadPosition(ReadSize);
+
+	return ReadSize;
+}
+
+void FBufferTransport::MarkComplete()
+{
+	SendBufferWriter.MarkComplete();
+}
+
+void FBufferTransport::Close()
+{
+	RecvBufferReader.Detach();
+}
\ No newline at end of file
diff --git a/Engine/Source/Developer/Horde/Public/AgentMessage.h b/Engine/Source/Developer/Horde/Public/AgentMessage.h
new file mode 100644
--- /dev/null
+++ b/Engine/Source/Developer/Horde/Public/AgentMessage.h
@@ -0,0 +1,183 @@
+// Copyright Epic Games, Inc. All Rights Reserved.
+
+#pragma once
+
+#include <memory>
+#include <vector>
+#include <map>
+#include <string_view>
+#include "ComputeBuffer.h"
+
+// Type of a compute message
+enum class EAgentMessageType : unsigned char
+{
+	// No message was received (end of stream)
+	None = 0x00,
+
+	// No-op message sent to keep the connection alive. Remote should reply with the same message.
+	Ping = 0x01,
+
+	// Sent in place of a regular response if an error occurs on the remote
+	Exception = 0x02,
+
+	// Fork the message loop into a new channel
+	Fork = 0x03,
+
+	// Sent as the first message on a channel to notify the remote that the remote end is attached
+	Attach = 0x04,
+
+	// Extract files on the remote machine (Initiator -> Remote)
+	WriteFiles = 0x10,
+
+	// Notification that files have been extracted (Remote -> Initiator)
+	WriteFilesResponse = 0x11,
+
+	// Deletes files on the remote machine (Initiator -> Remote)
+	DeleteFiles = 0x12,
+
+	// Execute a process in a sandbox (Initiator -> Remote)
+	ExecuteV1 = 0x16,
+
+	// Execute a process in a sandbox (Initiator -> Remote)
+	ExecuteV2 = 0x22,
+
+	// Returns output from the child process to the caller (Remote -> Initiator)
+	ExecuteOutput = 0x17,
+
+	// Returns the process exit code (Remote -> Initiator)
+	ExecuteResult = 0x18,
+
+	// Reads a blob from storage
+	ReadBlob = 0x20,
+
+	// Response to a ReadBlob request.
+	ReadBlobResponse = 0x21,
+
+	// Xor a block of data with a value
+	XorRequest = 0xf0,
+
+	// Result from an XorRequest request.
+	XorResponse = 0xf1,
+};
+
+// Flags describing how to execute a compute task process on the agent
+enum class EExecuteProcessFlags : unsigned char
+{
+	// No execute flags set
+	None = 0,
+
+	// Request execution to be wrapped under Wine when running on Linux. Agent still reserves the right to refuse it (e.g no Wine executable configured, mismatching OS etc)
+	UseWine = 1,
+};
+
+// Definitions for various message types
+namespace AgentMessage
+{
+	struct FException
+	{
+		std::string_view Message;
+		std::string_view Description;
+	};
+
+	struct FBlobRequest
+	{
+		std::string_view Locator;
+		int Offset;
+		int Length;
+	};
+}
+
+// Channel for receiving agent messages
+class FAgentMessageChannel
+{
+public:
+	FAgentMessageChannel(std::shared_ptr<FComputeBufferReader> InRecvBufferReader, std::shared_ptr<FComputeBufferWriter> InSendBufferWriter);
+	~FAgentMessageChannel();
+
+	//// Requests ////
+
+	// Closes the remote message loop
+	void Close();
+
+	// Sends a ping message to the remote
+	void Ping();
+
+	// Sends an exception response to the remote
+	void Exception(const char* Description, const char* Trace);
+
+	// Requests that the remote message loop be forked
+	void Fork(int ChannelId, int BufferSize);
+
+	// Notifies the remote that a buffer has been attached
+	void Attach();
+
+	// Extracts a bundle containing files to a particular path
+	void UploadFiles(const char* Path, const char* Locator);
+
+	// Deletes files matching a set of wildcards
+	void DeleteFiles(const char** Paths, size_t Count);
+
+	// Executes a process on the remote machine
+	void Execute(const char* Exe, const char** Args, size_t NumArgs, const char* WorkingDir, const char** EnvVars, size_t NumEnvVars, EExecuteProcessFlags Flags);
+
+	// Writes a blob requested by the remote
+	void Blob(const unsigned char* Data, size_t Length);
+
+	// Send a message to request that a byte string be xor'ed with a particular value
+	void Xor(const unsigned char* Data, size_t Length, unsigned char Value);
+
+	//// Responses ////
+
+	// Reads a response from the remote. Other Read methods can be used to access response data.
+	EAgentMessageType ReadResponse();
+
+	// Gets the raw data from a response
+	const void* GetResponseData() const { return ResponseData + MessageHeaderLength; }
+
+	// Gets the size of the response
+	size_t GetResponseSize() const { return ResponseLength - MessageHeaderLength; }
+
+	// Reads an exception response
+	void ReadException(AgentMessage::FException& Ex);
+
+	// Reads the result from executing a process
+	int ReadExecuteResult();
+
+	// Reads a blob request message
+	void ReadBlobRequest(AgentMessage::FBlobRequest& Ex);
+
+private:
+	const size_t MessageHeaderLength = 5;
+
+	std::shared_ptr<FComputeBufferReader> RecvBufferReader;
+	std::shared_ptr<FComputeBufferWriter> SendBufferWriter;
+
+	unsigned char* RequestData;
+	size_t RequestSize;
+	size_t MaxRequestSize;
+
+	EAgentMessageType ResponseType;
+	const unsigned char* ResponseData;
+	size_t ResponseLength;
+
+	void CreateMessage(EAgentMessageType Type, size_t MaxLength);
+	void FlushMessage();
+
+	void WriteInt32(int Value);
+	static int ReadInt32(const unsigned char** Pos);
+
+	void WriteFixedLengthBytes(const unsigned char* Data, size_t Length);
+	static const unsigned char* ReadFixedLengthBytes(const unsigned char** Pos, size_t Length);
+
+	static size_t MeasureUnsignedVarInt(size_t Value);
+	void WriteUnsignedVarInt(size_t Value);
+	static size_t ReadUnsignedVarInt(const unsigned char** Pos);
+
+	size_t MeasureString(const char* Text) const;
+	void WriteString(const char* Text);
+	void WriteString(const std::string_view& Text);
+	static std::string_view ReadString(const unsigned char** Pos);
+
+	void WriteOptionalString(const char* Text);
+};
+
diff --git a/Engine/Source/Developer/Horde/Public/ComputeBuffer.h b/Engine/Source/Developer/Horde/Public/ComputeBuffer.h
new file mode 100644
--- /dev/null
+++ b/Engine/Source/Developer/Horde/Public/ComputeBuffer.h
@@ -0,0 +1,174 @@
+// Copyright Epic Games, Inc. All Rights Reserved.
+
+#pragma once
+
+#include <memory>
+
+struct FHeader;
+struct FComputeBufferDetail;
+
+class FComputeBufferReader;
+class FComputeBufferWriter;
+
+//
+// Implements a ring buffer using shared memory, with one writer and multiple readers. 
+// 
+// Contention between the writer and readers is resolved by splitting the buffer into a number of
+// fixed-size chunks, which can be appended to and read immediately, but only reset once consumed
+// by all readers.
+//
+// If the read position is never advanced and only one chunk is used, the buffer functions as an
+// append-only buffer.
+//
+class FComputeBuffer
+{
+public:
+	// Maximum allowed number of readers
+	static const int MaxReaders = 16;
+
+	// Maximum allowed number of chunks
+	static const int MaxChunks = 16;
+
+	// Maximum length of a compute buffer name
+	static const size_t MaxNameLength = 256;
+
+	// Parameters for constructing a compute buffer
+	struct FParams
+	{
+		const char* Name = nullptr; // When creating a buffer, we should not normally set the name. We can only *OPEN* buffers with a name. Usually we create buffers and attach them to channels.
+		int NumChunks = 2;
+		int ChunkLength = 64 * 1024;
+		int NumReaders = 1;
+	};
+
+	FComputeBuffer();
+	FComputeBuffer(const FComputeBuffer& Other);
+	FComputeBuffer(FComputeBuffer&& Other) noexcept;
+	~FComputeBuffer();
+
+	FComputeBuffer& operator=(const FComputeBuffer& Other);
+	FComputeBuffer& operator=(FComputeBuffer&& Other) noexcept;
+
+	// Creates a new buffer
+	bool CreateNew(const FParams& Params);
+
+	// Opens an existing shared memory buffer (typically from handles created in another process)
+	bool OpenExisting(const char* Name);
+
+	// Close the current buffer and release all allocated resources
+	void Close();
+
+	// Test if the buffer is currently open
+	bool IsValid() const { return Detail != nullptr; }
+
+	// Creates a new reader for this buffer
+	FComputeBufferReader CreateReader();
+
+	// Creates a new writer for this buffer
+	FComputeBufferWriter CreateWriter();
+
+private:
+	friend class FWorkerComputeSocket;
+
+	FComputeBufferDetail* Detail;
+
+	const char* GetName() const;
+};
+
+//
+// Facilitates reading data from a compute buffer
+//
+class FComputeBufferReader
+{
+public:
+	FComputeBufferReader();
+	FComputeBufferReader(const FComputeBufferReader& Other);
+	FComputeBufferReader(FComputeBufferReader&& Other) noexcept;
+	~FComputeBufferReader();
+
+	FComputeBufferReader& operator=(const FComputeBufferReader& Other);
+	FComputeBufferReader& operator=(FComputeBufferReader&& Other) noexcept;
+
+	// Closes the handle to the underlying reader instance, resetting this instance back to empty
+	void Close();
+
+	// Detaches this reader from the buffer, causing all pending and subsequent reads to return immediately.
+	void Detach();
+
+	// Test if the reader is valid
+	bool IsValid() const { return Detail != nullptr; }
+
+	// Test whether the buffer has finished being written to (ie. MarkComplete() has been called by the writer) and all data has been read from it.
+	bool IsComplete() const;
+
+	// Move the read cursor forwards by the given number of bytes
+	void AdvanceReadPosition(size_t Size);
+
+	// Gets the amount of data that is ready to be read from a contiguous block of memory.
+	size_t GetMaxReadSize() const;
+
+	// Reads data into the given buffer
+	size_t Read(void* Buffer, size_t MaxSize, int TimeoutMs = -1);
+
+	// Waits until the given amount of data has been read, and returns a pointer to it. Returns nullptr if the timeout expires, or
+	// if the requested amount of data is not in a contiguous block of memory.
+	const unsigned char* WaitToRead(size_t MinSize, int TimeoutMs = -1);
+
+private:
+	struct FReaderRef;
+
+	friend class FComputeBuffer;
+	friend class FWorkerComputeSocket;
+
+	FComputeBufferDetail* Detail;
+	int ReaderIdx;
+
+	FComputeBufferReader(FComputeBufferDetail* Detail, int ReaderIdx);
+	const char* GetName() const;
+};
+
+//
+// Facilitates writing data to a compute buffer
+//
+class FComputeBufferWriter
+{
+public:
+	FComputeBufferWriter();
+	FComputeBufferWriter(const FComputeBufferWriter& Other);
+	FComputeBufferWriter(FComputeBufferWriter&& Other) noexcept;
+	~FComputeBufferWriter();
+
+	FComputeBufferWriter& operator=(const FComputeBufferWriter& Other);
+	FComputeBufferWriter& operator=(FComputeBufferWriter&& Other) noexcept;
+
+	// Closes the handle to the underlying writer instance, resetting this instance back to empty
+	void Close();
+
+	// Test if the writer is valid
+	bool IsValid() const { return Detail != nullptr; }
+
+	// Signal that we've finished writing to this buffer
+	void MarkComplete();
+
+	// Move the write cursor forward by the given number of bytes
+	void AdvanceWritePosition(size_t Size);
+
+	// Gets the length of the current write buffer.
+	size_t GetMaxWriteSize() const;
+
+	// Writes data into the compute buffer
+	size_t Write(const void* Buffer, size_t MaxSize, int TimeoutMs = -1);
+
+	// Waits until a write buffer of the requested size is available, and returns a pointer to it. Returns nullptr if the 
+	// timeout expires before enough data has been flushed to return a buffer of the given length.
+	unsigned char* WaitToWrite(size_t MinSize, int TimeoutMs = -1);
+
+private:
+	friend class FComputeBuffer;
+	friend class FWorkerComputeSocket;
+
+	FComputeBufferDetail* Detail;
+
+	FComputeBufferWriter(FComputeBufferDetail* Detail);
+	const char* GetName() const;
+};
diff --git a/Engine/Source/Developer/Horde/Public/ComputeChannel.h b/Engine/Source/Developer/Horde/Public/ComputeChannel.h
new file mode 100644
--- /dev/null
+++ b/Engine/Source/Developer/Horde/Public/ComputeChannel.h
@@ -0,0 +1,36 @@
+// Copyright Epic Games, Inc. All Rights Reserved.
+
+#pragma once
+
+#include "ComputeBuffer.h"
+
+class FComputeSocket;
+
+//
+// Allows bi-directional communication between two nodes using compute buffers
+//
+class FComputeChannel
+{
+public:
+	// Reader for the channel
+	FComputeBufferReader Reader;
+
+	// Writer for the channel
+	FComputeBufferWriter Writer;
+
+	FComputeChannel();
+	FComputeChannel(FComputeBufferReader InReader, FComputeBufferWriter InWriter);
+	~FComputeChannel();
+
+	// Tests whether the channel is valid
+	bool IsValid() const;
+
+	// Sends bytes to the remote. 
+	size_t Send(const void* Data, size_t Size, int TimeoutMs = -1);
+
+	// Reads as many bytes as are available from the socket.
+	size_t Recv(void* Data, size_t Size, int TimeoutMs = -1);
+
+	// Indicate to the remote that no more data will be sent.
+	void MarkComplete();
+};
diff --git a/Engine/Source/Developer/Horde/Public/ComputePlatform.h b/Engine/Source/Developer/Horde/Public/ComputePlatform.h
new file mode 100644
--- /dev/null
+++ b/Engine/Source/Developer/Horde/Public/ComputePlatform.h
@@ -0,0 +1,165 @@
+// Copyright Epic Games, Inc. All Rights Reserved.
+
+#pragma once
+
+#include <stdio.h>
+#include <stddef.h>
+#include <stdint.h>
+
+// Defines for the current platform
+#ifdef _MSC_VER
+	#define UE_COMPUTE_PLATFORM_WINDOWS 1
+	#define UE_COMPUTE_PLATFORM_MAC 0
+	#define UE_COMPUTE_PLATFORM_LINUX 0
+#elif defined(__APPLE__)
+	#define UE_COMPUTE_PLATFORM_WINDOWS 0
+	#define UE_COMPUTE_PLATFORM_MAC 1
+	#define UE_COMPUTE_PLATFORM_LINUX 0
+#else
+	#define UE_COMPUTE_PLATFORM_WINDOWS 0
+	#define UE_COMPUTE_PLATFORM_MAC 0
+	#define UE_COMPUTE_PLATFORM_LINUX 1
+#endif
+
+#ifdef NDEBUG
+	#define UE_COMPUTE_ASSERT(x) ((void)0)
+	#define UE_COMPUTE_VERIFY(x) (!!(x))
+#else
+	#define UE_COMPUTE_ASSERT(x) ((void)UE_COMPUTE_VERIFY(x))
+	#define UE_COMPUTE_VERIFY(x) (!!(x) || (FComputePlatform::AssertFailed(#x, __FILE__, __LINE__), true))
+#endif
+
+//
+// IPC manually-reset event object
+//
+class FComputeEvent
+{
+public:
+	FComputeEvent();
+	~FComputeEvent();
+
+	// Creates a new event with the given name
+	bool Create(const char* Name);
+
+	// Opens an existing event created elsewhere
+	bool OpenExisting(const char* Name);
+
+	// Close the event and release its resources
+	void Close();
+
+	// Signal the event, releasing any waiters
+	void Signal();
+
+	// Wait for the event to be signalled or timeout. Pass -1 for timeoutMs to wait infinitely.
+	bool Wait(int timeoutMs);
+
+private:
+	void* Handle;
+};
+
+//
+// IPC memory mapped file
+//
+class FComputeMemoryMappedFile
+{
+public:
+	FComputeMemoryMappedFile();
+	~FComputeMemoryMappedFile();
+
+	// Creates a new memory mapped file with the given capacity
+	bool Create(const char* Name, long long Capacity);
+
+	// Opens an existing memory mapped file
+	bool OpenExisting(const char* Name);
+
+	// Close the memory mapped file handle
+	void Close();
+
+	// Gets a pointer to the mapped memory
+	void* GetPointer() const;
+
+private:
+	void* Handle;
+	void* Pointer;
+	long long MappedSize;
+	char* OwnerName;
+};
+
+//
+// Platform-specific utility functions
+//
+struct FComputePlatform
+{
+	//
+	// General
+	//
+
+	// Signal that an assertion has failed
+	static void AssertFailed(const char* Expr, const char* File, int Line);
+	
+	// Reads an environment variable
+	static bool GetEnvironmentVariable(const char* Name, char* Buffer, size_t BufferLen);
+
+	// Creates a unique object name
+	static void CreateUniqueName(char* NameBuffer, size_t NameBufferLen);
+
+	//
+	// Math
+	//
+
+	// Find the log2 of the given value, returning 0 if the value is zero.
+	static unsigned int FloorLog2(unsigned int Value);
+
+	// Count the number of leading zeros in the given value.
+	static unsigned int CountLeadingZeros(unsigned int Value);
+
+	//
+	// Strings
+	//
+
+	// Copy a string from one buffer to another, not exceeding the destination buffer size
+	static void Strcpy(char* Dest, size_t DestLen, const char* Source);
+
+	// Perform a case-insensitive comparison of two strings
+	static int Stricmp(const char* A, const char* B);
+
+	//
+	// Atomics
+	//
+
+	// Reads a 64-bit value from the given memory location 
+	static long long AtomicRead64(const volatile long long* Ptr);
+
+	// Writes a 64-bit value to the given memory location
+	static void AtomicWrite64(volatile long long* Ptr, long long Value);
+
+	// Increments an integer at the given memory location and returns the incremented value.
+	static long AtomicIncrement(volatile long* Ptr);
+
+	// Increments an integer at the given memory location and returns the incremented value.
+	static long long AtomicIncrement64(volatile long long* Ptr);
+
+	// Decrements an integer at the given memory location and returns the decremented value.
+	static long AtomicDecrement(volatile long* Ptr);
+
+	// Adds a value to the given memory location, and returns the resulting value.
+	static long long AtomicAdd64(volatile long long* Ptr, long long Value);
+
+	// ANDs a value with a given memory location, and returns the contents of the memory location BEFORE the operation.
+	static long AtomicAnd(volatile long* Ptr, long Value);
+
+	// ANDs a value with a given memory location, and returns the contents of the memory location BEFORE the operation.
+	static long long AtomicAnd64(volatile long long* Ptr, long long Value);
+
+	// ORs a value with a given memory location, and returns the contents of the memory location BEFORE the operation.
+	static long long AtomicOr64(volatile long long* Ptr, long long Value);
+
+	// Sets the contents of a pointer to the given value, and returns the original value in that memory location.
+	static long long AtomicExchange64(volatile long long* Ptr, long long Exchange);
+
+	// Compares the contents of a memory location to a value, and exchanges it for another value if they are equal. Returns true if the operation succeeded.
+	static bool AtomicCompareExchange(volatile long* Ptr, long Exchange, long Comperand);
+
+	// Compares the contents of a memory location to a value, and exchanges it for another value if they are equal. Returns true if the operation succeeded.
+	static bool AtomicCompareExchange64(volatile long long* Ptr, long long Exchange, long long Comperand);
+};
diff --git a/Engine/Source/Developer/Horde/Public/ComputeSocket.h b/Engine/Source/Developer/Horde/Public/ComputeSocket.h
new file mode 100644
--- /dev/null
+++ b/Engine/Source/Developer/Horde/Public/ComputeSocket.h
@@ -0,0 +1,94 @@
+// Copyright Epic Games, Inc. All Rights Reserved.
+
+#pragma once
+
+#include "ComputeBuffer.h"
+#include "ComputeChannel.h"
+#include "ComputeTransport.h"
+#include <vector>
+
+//
+// Connection to a remote machine that multiplexes data into and out-of multiple buffers
+// attached to different channel numbers.
+//
+class FComputeSocket
+{
+public:
+	FComputeSocket();
+	virtual ~FComputeSocket();
+
+	FComputeSocket(const FComputeSocket&) = delete;
+	FComputeSocket& operator=(const FComputeSocket&) = delete;
+
+	// Attaches a new buffer for receiving data
+	virtual void AttachRecvBuffer(int ChannelId, FComputeBuffer RecvBuffer) = 0;
+
+	// Attaches a new buffer for sending data */
+	virtual void AttachSendBuffer(int ChannelId, FComputeBuffer SendBuffer) = 0;
+
+	// Attaches a channel to this socket
+	FComputeChannel CreateChannel(int ChannelId);
+
+	// Attaches a channel to this socket
+	FComputeChannel CreateChannel(int ChannelId, FComputeBuffer RecvBuffer, FComputeBuffer SendBuffer);
+};
+
+//
+// Socket used by a worker process to communicate with a host running on the same machine
+// using shared memory to attach new buffers.
+//
+class FWorkerComputeSocket final : public FComputeSocket
+{
+public:
+	static const char* const IpcEnvVar;
+
+	FWorkerComputeSocket();
+	~FWorkerComputeSocket();
+
+	// Opens a connection to the agent process using a command buffer read from an environment variable (EnvVarName)
+	bool Open();
+
+	// Opens a connection to the agent process using a specific command buffer name
+	bool Open(const char* CommandBufferName);
+
+	// Close the current connection
+	void Close();
+
+	// Attaches a new buffer for receiving data
+	virtual void AttachRecvBuffer(int ChannelId, FComputeBuffer RecvBuffer) override;
+
+	// Attaches a new buffer for sending data
+	virtual void AttachSendBuffer(int ChannelId, FComputeBuffer SendBuffer) override;
+
+	// Reads and handles a command from the command buffer
+	static void RunServer(FComputeBufferReader& CommandBufferReader, FComputeSocket& Socket);
+
+private:
+	enum class EMessageType;
+
+	FComputeBufferWriter CommandBufferWriter;
+	std::vector<FComputeBuffer> Buffers;
+
+	void AttachBuffer(int ChannelId, EMessageType Type, const char* Name);
+
+	static size_t ReadVarUInt(const unsigned char* Pos, unsigned int* OutValue);
+	static size_t ReadString(const unsigned char* Pos, char* OutText, size_t OutTextMaxLen);
+
+	static size_t WriteVarUInt(unsigned char* Pos, unsigned int Value);
+	static size_t WriteString(unsigned char* Pos, const char* Text);
+};
+
+//
+// Enum identifying which end of the socket a particular machine is
+//
+enum class EComputeSocketEndpoint
+{
+	// The initiating machine
+	Local,
+
+	// The remote machine
+	Remote
+};
+
+// Creates a socket using a custom transport
+std::unique_ptr<FComputeSocket> CreateComputeSocket(std::unique_ptr<FComputeTransport> Transport, EComputeSocketEndpoint Endpoint);
diff --git a/Engine/Source/Developer/Horde/Public/ComputeTransport.h b/Engine/Source/Developer/Horde/Public/ComputeTransport.h
new file mode 100644
--- /dev/null
+++ b/Engine/Source/Developer/Horde/Public/ComputeTransport.h
@@ -0,0 +1,51 @@
+// Copyright Epic Games, Inc. All Rights Reserved.
+
+#pragma once
+
+#include "ComputeBuffer.h"
+
+//
+// Interface for a transport mechanism that can be used by a compute socket
+//
+class FComputeTransport
+{
+public:
+	virtual ~FComputeTransport();
+	
+	// Sends data to the remote
+	virtual size_t Send(const void* Data, size_t Size) = 0;
+
+	// Receives data from the remote
+	virtual size_t Recv(void* Data, size_t Size) = 0;
+
+	// Indicates to the remote that no more data will be sent.
+	virtual void MarkComplete() = 0;
+
+	// Indicates that no more data will be sent or received, and that any blocking reads/writes should stop.
+	virtual void Close() = 0;
+
+	// Sends data to the remote, blocking until the entire message has been sent.
+	bool SendMessage(const void* Data, size_t Size);
+
+	// Receives a fixed length block of data from the remote, blocking until the entire length has been received.
+	bool RecvMessage(void* Data, size_t Size);
+};
+
+//
+// Implementation of FComputeTransport which uses in-memory buffers to transport data
+//
+class FBufferTransport final : public FComputeTransport
+{
+public:
+	FBufferTransport(FComputeBufferWriter InSendBufferWriter, FComputeBufferReader InRecvBufferReader);
+
+protected:
+	virtual size_t Send(const void* Data, size_t Size) override final;
+	virtual size_t Recv(void* Data, size_t Size) override final;
+	virtual void MarkComplete() override final;
+	virtual void Close() override final;
+
+private:
+	FComputeBufferWriter SendBufferWriter;
+	FComputeBufferReader RecvBufferReader;
+};
diff --git a/Engine/Source/Programs/Horde/Samples/RemoteWorkerCpp/AgentMessage.cpp b/Engine/Source/Programs/Horde/Samples/RemoteWorkerCpp/AgentMessage.cpp
deleted file mode 100644
diff --git a/Engine/Source/Programs/Horde/Samples/RemoteWorkerCpp/AgentMessage.h b/Engine/Source/Programs/Horde/Samples/RemoteWorkerCpp/AgentMessage.h
deleted file mode 100644
diff --git a/Engine/Source/Programs/Horde/Samples/RemoteWorkerCpp/ComputeBuffer.cpp b/Engine/Source/Programs/Horde/Samples/RemoteWorkerCpp/ComputeBuffer.cpp
deleted file mode 100644
diff --git a/Engine/Source/Programs/Horde/Samples/RemoteWorkerCpp/ComputeBuffer.h b/Engine/Source/Programs/Horde/Samples/RemoteWorkerCpp/ComputeBuffer.h
deleted file mode 100644
diff --git a/Engine/Source/Programs/Horde/Samples/RemoteWorkerCpp/ComputeChannel.cpp b/Engine/Source/Programs/Horde/Samples/RemoteWorkerCpp/ComputeChannel.cpp
deleted file mode 100644
diff --git a/Engine/Source/Programs/Horde/Samples/RemoteWorkerCpp/ComputeChannel.h b/Engine/Source/Programs/Horde/Samples/RemoteWorkerCpp/ComputeChannel.h
deleted file mode 100644
diff --git a/Engine/Source/Programs/Horde/Samples/RemoteWorkerCpp/ComputePlatform.cpp b/Engine/Source/Programs/Horde/Samples/RemoteWorkerCpp/ComputePlatform.cpp
deleted file mode 100644
diff --git a/Engine/Source/Programs/Horde/Samples/RemoteWorkerCpp/ComputePlatform.h b/Engine/Source/Programs/Horde/Samples/RemoteWorkerCpp/ComputePlatform.h
deleted file mode 100644
diff --git a/Engine/Source/Programs/Horde/Samples/RemoteWorkerCpp/ComputeSocket.cpp b/Engine/Source/Programs/Horde/Samples/RemoteWorkerCpp/ComputeSocket.cpp
deleted file mode 100644
diff --git a/Engine/Source/Programs/Horde/Samples/RemoteWorkerCpp/ComputeSocket.h b/Engine/Source/Programs/Horde/Samples/RemoteWorkerCpp/ComputeSocket.h
deleted file mode 100644
diff --git a/Engine/Source/Programs/Horde/Samples/RemoteWorkerCpp/ComputeTransport.cpp b/Engine/Source/Programs/Horde/Samples/RemoteWorkerCpp/ComputeTransport.cpp
deleted file mode 100644
diff --git a/Engine/Source/Programs/Horde/Samples/RemoteWorkerCpp/ComputeTransport.h b/Engine/Source/Programs/Horde/Samples/RemoteWorkerCpp/ComputeTransport.h
deleted file mode 100644
