Horde: Static analyzer fix.

diff --git a/Engine/Source/Developer/Horde/Private/ComputePlatform.cpp b/Engine/Source/Developer/Horde/Private/ComputePlatform.cpp
--- a/Engine/Source/Developer/Horde/Private/ComputePlatform.cpp
+++ b/Engine/Source/Developer/Horde/Private/ComputePlatform.cpp
@@ -353,7 +353,7 @@ void FComputePlatform::CreateUniqueName(char* NameBuffer, size_t NameBufferLen)
 #if UE_COMPUTE_PLATFORM_WINDOWS
 	DWORD Pid = GetCurrentProcessId();
 	ULONGLONG TickCount = GetTickCount64();
-	snprintf(NameBuffer, NameBufferLen, "Local\\COMPUTE_%lu_%llu_%lu", Pid, TickCount, AtomicIncrement(&Counter));
+	snprintf(NameBuffer, NameBufferLen, "Local\\COMPUTE_%lu_%llu_%lu", Pid, TickCount, (unsigned long)AtomicIncrement(&Counter));
 #else
 	struct timespec ts;
 	clock_gettime(CLOCK_MONOTONIC, &ts);
