Fixed small bugs in the Horde C++ code (for instance, int parsing, some message serialization).

diff --git a/Engine/Source/Developer/Horde/Private/AgentMessage.cpp b/Engine/Source/Developer/Horde/Private/AgentMessage.cpp
--- a/Engine/Source/Developer/Horde/Private/AgentMessage.cpp
+++ b/Engine/Source/Developer/Horde/Private/AgentMessage.cpp
@@ -5,9 +5,8 @@
 
 //// FAgentMessageChannel ////
 
-FAgentMessageChannel::FAgentMessageChannel(std::shared_ptr<FComputeBufferReader> InRecvBufferReader, std::shared_ptr<FComputeBufferWriter> InSendBufferWriter)
-	: RecvBufferReader(std::move(InRecvBufferReader))
-	, SendBufferWriter(std::move(InSendBufferWriter))
+FAgentMessageChannel::FAgentMessageChannel(std::shared_ptr<FComputeChannel> InChannel)
+	: ChannelBuffers(InChannel)
 	, RequestData(0)
 	, RequestSize(0)
 	, MaxRequestSize(0)
@@ -130,7 +129,7 @@ void FAgentMessageChannel::Execute(const char* Exe, const char** Args, size_t Nu
 
 void FAgentMessageChannel::Blob(const unsigned char* Data, size_t Length)
 {
-	const size_t MaxChunkSize = 512 * 1024;
+	const size_t MaxChunkSize = ChannelBuffers->Writer.GetChunkMaxLength() - 128 - MessageHeaderLength;
 	for (size_t ChunkOffset = 0; ChunkOffset < Length;)
 	{
 		size_t ChunkLength = std::min(Length - ChunkOffset, MaxChunkSize);
@@ -155,17 +154,17 @@ void FAgentMessageChannel::Xor(const unsigned char* Data, size_t Length, unsigne
 
 EAgentMessageType FAgentMessageChannel::ReadResponse()
 {
-	if (ResponseLength > 0)
+	if (ResponseData)
 	{
-		RecvBufferReader->AdvanceReadPosition(ResponseLength);
+		ChannelBuffers->Reader.AdvanceReadPosition(ResponseLength + MessageHeaderLength);
 		ResponseData = nullptr;
 		ResponseLength = 0;
 	}
 
-	const unsigned char* Header = RecvBufferReader->WaitToRead(MessageHeaderLength);
+	const unsigned char* Header = ChannelBuffers->Reader.WaitToRead(MessageHeaderLength);
 	unsigned int Length = *((unsigned int*)(Header + 1));
 
-	Header = RecvBufferReader->WaitToRead(MessageHeaderLength + Length);
+	Header = ChannelBuffers->Reader.WaitToRead(MessageHeaderLength + Length);
 
 	ResponseType = (EAgentMessageType)Header[0];
 	ResponseData = Header + MessageHeaderLength;
@@ -208,14 +207,14 @@ void FAgentMessageChannel::ReadBlobRequest(AgentMessage::FBlobRequest& Ex)
 
 void FAgentMessageChannel::CreateMessage(EAgentMessageType Type, size_t MaxLength)
 {
-	RequestData = SendBufferWriter->WaitToWrite(MessageHeaderLength + MaxLength);
+	RequestData = ChannelBuffers->Writer.WaitToWrite(MessageHeaderLength + MaxLength);
 	RequestData[0] = (unsigned char)Type;
 }
 
 void FAgentMessageChannel::FlushMessage()
 {
 	memcpy(&RequestData[1], &RequestSize, sizeof(int));
-	SendBufferWriter->AdvanceWritePosition(MessageHeaderLength + RequestSize);
+	ChannelBuffers->Writer.AdvanceWritePosition(MessageHeaderLength + RequestSize);
 	RequestSize = 0;
 	MaxRequestSize = 0;
 	RequestData = nullptr;
@@ -251,7 +250,15 @@ const unsigned char* FAgentMessageChannel::ReadFixedLengthBytes(const unsigned c
 size_t FAgentMessageChannel::MeasureUnsignedVarInt(size_t Value)
 {
 	UE_COMPUTE_ASSERT(Value == (unsigned int)Value);
-	return (FComputePlatform::FloorLog2((unsigned int)Value) / 7) + 1;
+
+	if (Value == 0)
+	{
+		return 1;
+	}
+	else
+	{
+		return (FComputePlatform::FloorLog2((unsigned int)Value) / 7) + 1;
+	}
 }
 
 void FAgentMessageChannel::WriteUnsignedVarInt(size_t Value)
@@ -275,7 +282,7 @@ size_t FAgentMessageChannel::ReadUnsignedVarInt(const unsigned char** Pos)
 	const unsigned char* Data = *Pos;
 
 	unsigned char FirstByte = Data[0];
-	size_t NumBytes = FComputePlatform::CountLeadingZeros(~(unsigned int)FirstByte) + 1 - 24; // Note byte -> int conversion here, hence ignoring subtracting 24 bits
+	size_t NumBytes = FComputePlatform::CountLeadingZeros(0xFF & (~(unsigned int)FirstByte)) + 1 - 24; // Note byte -> int conversion here, hence ignoring subtracting 24 bits
 
 	size_t value = (size_t)(FirstByte & (0xff >> NumBytes));
 	for (size_t Idx = 1; Idx < NumBytes; Idx++)
diff --git a/Engine/Source/Developer/Horde/Private/ComputeBuffer.cpp b/Engine/Source/Developer/Horde/Private/ComputeBuffer.cpp
--- a/Engine/Source/Developer/Horde/Private/ComputeBuffer.cpp
+++ b/Engine/Source/Developer/Horde/Private/ComputeBuffer.cpp
@@ -1153,6 +1153,11 @@ size_t FComputeBufferWriter::GetMaxWriteSize() const
 	return Detail->GetMaxWriteSize();
 }
 
+size_t FComputeBufferWriter::GetChunkMaxLength() const
+{
+	return Detail->Header->ChunkLength;
+}
+
 size_t FComputeBufferWriter::Write(const void* Buffer, size_t MaxSize, int TimeoutMs)
 {
 	unsigned char* SendData = WaitToWrite(1, TimeoutMs);
diff --git a/Engine/Source/Developer/Horde/Private/ComputeSocket.cpp b/Engine/Source/Developer/Horde/Private/ComputeSocket.cpp
--- a/Engine/Source/Developer/Horde/Private/ComputeSocket.cpp
+++ b/Engine/Source/Developer/Horde/Private/ComputeSocket.cpp
@@ -9,6 +9,7 @@
 #include <vector>
 #include <thread>
 #include <mutex>
+#include <chrono>
 
 FComputeSocket::FComputeSocket()
 {
@@ -18,26 +19,26 @@ FComputeSocket::~FComputeSocket()
 {
 }
 
-FComputeChannel FComputeSocket::CreateChannel(int ChannelId)
+std::shared_ptr<FComputeChannel> FComputeSocket::CreateChannel(int ChannelId)
 {
 	FComputeBuffer RecvBuffer;
 	if (!RecvBuffer.CreateNew(FComputeBuffer::FParams()))
 	{
-		return FComputeChannel();
+		return std::make_shared<FComputeChannel>(FComputeChannel());
 	}
 
 	FComputeBuffer SendBuffer;
 	if (!SendBuffer.CreateNew(FComputeBuffer::FParams()))
 	{
-		return FComputeChannel();
+		return std::make_shared<FComputeChannel>(FComputeChannel());
 	}
 
 	return CreateChannel(ChannelId, std::move(RecvBuffer), std::move(SendBuffer));
 }
 
-FComputeChannel FComputeSocket::CreateChannel(int ChannelId, FComputeBuffer RecvBuffer, FComputeBuffer SendBuffer)
+std::shared_ptr<FComputeChannel> FComputeSocket::CreateChannel(int ChannelId, FComputeBuffer RecvBuffer, FComputeBuffer SendBuffer)
 {
-	FComputeChannel Channel(RecvBuffer.CreateReader(), SendBuffer.CreateWriter());
+	std::shared_ptr<FComputeChannel> Channel = std::make_shared<FComputeChannel>(RecvBuffer.CreateReader(), SendBuffer.CreateWriter());
 
 	AttachRecvBuffer(ChannelId, std::move(RecvBuffer));
 	AttachSendBuffer(ChannelId, std::move(SendBuffer));
@@ -64,6 +65,10 @@ FWorkerComputeSocket::~FWorkerComputeSocket()
 	Close();
 }
 
+void FWorkerComputeSocket::StartCommunication()
+{
+}
+
 bool FWorkerComputeSocket::Open()
 {
 	char EnvVar[FComputeBuffer::MaxNameLength];
@@ -259,6 +264,11 @@ public:
 	const EComputeSocketEndpoint Endpoint;
 	std::mutex CriticalSection;
 
+	bool bPingThreadFinish;
+	std::mutex PingThreadFinishMutex;
+	std::condition_variable PingThreadFinishCV;
+	std::thread PingThread;
+
 	std::thread RecvThread;
 
 	std::unordered_map<int, FComputeBufferWriter> Writers;
@@ -269,12 +279,19 @@ public:
 		: Transport(std::move(InTransport))
 		, Endpoint(InEndpoint)
 		, CriticalSection()
-		, RecvThread(&FRemoteComputeSocket::RecvThreadProc, this)
+		, bPingThreadFinish(0)
 	{
 	}
 
-	~FRemoteComputeSocket()
+	~FRemoteComputeSocket() override
 	{
+		{
+			std::lock_guard<std::mutex>	Lock(PingThreadFinishMutex);
+			bPingThreadFinish = 1;
+		}
+
+		PingThreadFinishCV.notify_all();
+
 		for (FComputeBufferReader& Reader : Readers)
 		{
 			Reader.Detach();
@@ -287,6 +304,38 @@ public:
 
 		Transport->Close();
 		RecvThread.join();
+		PingThread.join();
+	}
+
+	virtual void StartCommunication()
+	{
+		bPingThreadFinish = 0;
+
+		// Initialize the receiver thread after having attached channel 0
+		RecvThread = std::thread(&FRemoteComputeSocket::RecvThreadProc, this);
+		PingThread = std::thread(&FRemoteComputeSocket::PingThreadProc, this);
+	}
+
+	void PingThreadProc()
+	{
+		for (;;)
+		{
+			{ // Send the ping message
+				std::lock_guard<std::mutex> Lock(CriticalSection);
+
+				FFrameHeader Header;
+				Header.Channel = 0;
+				Header.Size = -3; // Ping control message.
+
+				Transport->SendMessage(&Header, sizeof(Header));
+			}
+
+			std::unique_lock<std::mutex> Lock(PingThreadFinishMutex);
+			if (PingThreadFinishCV.wait_for(Lock, std::chrono::seconds(2), [this]() { return bPingThreadFinish; }))
+			{
+				break;
+			}
+		}
 	}
 
 	void RecvThreadProc()
diff --git a/Engine/Source/Developer/Horde/Public/AgentMessage.h b/Engine/Source/Developer/Horde/Public/AgentMessage.h
--- a/Engine/Source/Developer/Horde/Public/AgentMessage.h
+++ b/Engine/Source/Developer/Horde/Public/AgentMessage.h
@@ -6,7 +6,9 @@
 #include <vector>
 #include <map>
 #include <string_view>
+#include "Horde.h"
 #include "ComputeBuffer.h"
+#include "ComputeChannel.h"
 
 // Type of a compute message
 enum class EAgentMessageType : unsigned char
@@ -91,66 +93,65 @@ namespace AgentMessage
 class FAgentMessageChannel
 {
 public:
-	FAgentMessageChannel(std::shared_ptr<FComputeBufferReader> InRecvBufferReader, std::shared_ptr<FComputeBufferWriter> InSendBufferWriter);
-	~FAgentMessageChannel();
+	HORDE_API FAgentMessageChannel(std::shared_ptr<FComputeChannel> InChannel);
+	HORDE_API ~FAgentMessageChannel();
 
 	//// Requests ////
 
 	// Closes the remote message loop
-	void Close();
+	HORDE_API void Close();
 
 	// Sends a ping message to the remote
-	void Ping();
+	HORDE_API void Ping();
 
 	// Sends an exception response to the remote
-	void Exception(const char* Description, const char* Trace);
+	HORDE_API void Exception(const char* Description, const char* Trace);
 
 	// Requests that the remote message loop be forked
-	void Fork(int ChannelId, int BufferSize);
+	HORDE_API void Fork(int ChannelId, int BufferSize);
 
 	// Notifies the remote that a buffer has been attached
-	void Attach();
+	HORDE_API void Attach();
 
 	// Extracts a bundle containing files to a particular path
-	void UploadFiles(const char* Path, const char* Locator);
+	HORDE_API void UploadFiles(const char* Path, const char* Locator);
 
 	// Deletes files matching a set of wildcards
-	void DeleteFiles(const char** Paths, size_t Count);
+	HORDE_API void DeleteFiles(const char** Paths, size_t Count);
 
 	// Executes a process on the remote machine
-	void Execute(const char* Exe, const char** Args, size_t NumArgs, const char* WorkingDir, const char** EnvVars, size_t NumEnvVars, EExecuteProcessFlags Flags);
+	HORDE_API void Execute(const char* Exe, const char** Args, size_t NumArgs, const char* WorkingDir, const char** EnvVars, size_t NumEnvVars, EExecuteProcessFlags Flags);
 
 	// Writes a blob requested by the remote
-	void Blob(const unsigned char* Data, size_t Length);
+	HORDE_API void Blob(const unsigned char* Data, size_t Length);
 
 	// Send a message to request that a byte string be xor'ed with a particular value
-	void Xor(const unsigned char* Data, size_t Length, unsigned char Value);
+	HORDE_API void Xor(const unsigned char* Data, size_t Length, unsigned char Value);
 
 	//// Responses ////
 
 	// Reads a response from the remote. Other Read methods can be used to access response data.
-	EAgentMessageType ReadResponse();
+	HORDE_API EAgentMessageType ReadResponse();
 
 	// Gets the raw data from a response
-	const void* GetResponseData() const { return ResponseData + MessageHeaderLength; }
+	HORDE_API const void* GetResponseData() const { return ResponseData + MessageHeaderLength; }
 
 	// Gets the size of the response
-	size_t GetResponseSize() const { return ResponseLength - MessageHeaderLength; }
+	HORDE_API size_t GetResponseSize() const { return ResponseLength - MessageHeaderLength; }
 
 	// Reads an exception response
-	void ReadException(AgentMessage::FException& Ex);
+	HORDE_API void ReadException(AgentMessage::FException& Ex);
 
 	// Reads the result from executing a process
-	int ReadExecuteResult();
+	HORDE_API int ReadExecuteResult();
 
 	// Reads a blob request message
-	void ReadBlobRequest(AgentMessage::FBlobRequest& Ex);
+	HORDE_API void ReadBlobRequest(AgentMessage::FBlobRequest& Ex);
 
 private:
 	const size_t MessageHeaderLength = 5;
 
-	std::shared_ptr<FComputeBufferReader> RecvBufferReader;
-	std::shared_ptr<FComputeBufferWriter> SendBufferWriter;
+	std::shared_ptr<FComputeChannel> ChannelBuffers;
 
 	unsigned char* RequestData;
 	size_t RequestSize;
diff --git a/Engine/Source/Developer/Horde/Public/ComputeBuffer.h b/Engine/Source/Developer/Horde/Public/ComputeBuffer.h
--- a/Engine/Source/Developer/Horde/Public/ComputeBuffer.h
+++ b/Engine/Source/Developer/Horde/Public/ComputeBuffer.h
@@ -2,6 +2,7 @@
 
 #pragma once
 
+#include "Horde.h"
 #include <memory>
 
 struct FHeader;
@@ -37,35 +38,35 @@ public:
 	{
 		const char* Name = nullptr; // When creating a buffer, we should not normally set the name. We can only *OPEN* buffers with a name. Usually we create buffers and attach them to channels.
 		int NumChunks = 2;
-		int ChunkLength = 64 * 1024;
+		int ChunkLength = 512 * 1024;
 		int NumReaders = 1;
 	};
 
-	FComputeBuffer();
-	FComputeBuffer(const FComputeBuffer& Other);
-	FComputeBuffer(FComputeBuffer&& Other) noexcept;
-	~FComputeBuffer();
+	HORDE_API FComputeBuffer();
+	HORDE_API FComputeBuffer(const FComputeBuffer& Other);
+	HORDE_API FComputeBuffer(FComputeBuffer&& Other) noexcept;
+	HORDE_API ~FComputeBuffer();
 
-	FComputeBuffer& operator=(const FComputeBuffer& Other);
-	FComputeBuffer& operator=(FComputeBuffer&& Other) noexcept;
+	HORDE_API FComputeBuffer& operator=(const FComputeBuffer& Other);
+	HORDE_API FComputeBuffer& operator=(FComputeBuffer&& Other) noexcept;
 
 	// Creates a new buffer
-	bool CreateNew(const FParams& Params);
+	HORDE_API bool CreateNew(const FParams& Params);
 
 	// Opens an existing shared memory buffer (typically from handles created in another process)
-	bool OpenExisting(const char* Name);
+	HORDE_API bool OpenExisting(const char* Name);
 
 	// Close the current buffer and release all allocated resources
-	void Close();
+	HORDE_API void Close();
 
 	// Test if the buffer is currently open
-	bool IsValid() const { return Detail != nullptr; }
+	HORDE_API bool IsValid() const { return Detail != nullptr; }
 
 	// Creates a new reader for this buffer
-	FComputeBufferReader CreateReader();
+	HORDE_API FComputeBufferReader CreateReader();
 
 	// Creates a new writer for this buffer
-	FComputeBufferWriter CreateWriter();
+	HORDE_API FComputeBufferWriter CreateWriter();
 
 private:
 	friend class FWorkerComputeSocket;
@@ -81,38 +82,38 @@ private:
 class FComputeBufferReader
 {
 public:
-	FComputeBufferReader();
-	FComputeBufferReader(const FComputeBufferReader& Other);
-	FComputeBufferReader(FComputeBufferReader&& Other) noexcept;
-	~FComputeBufferReader();
+	HORDE_API FComputeBufferReader();
+	HORDE_API FComputeBufferReader(const FComputeBufferReader& Other);
+	HORDE_API FComputeBufferReader(FComputeBufferReader&& Other) noexcept;
+	HORDE_API ~FComputeBufferReader();
 
-	FComputeBufferReader& operator=(const FComputeBufferReader& Other);
-	FComputeBufferReader& operator=(FComputeBufferReader&& Other) noexcept;
+	HORDE_API FComputeBufferReader& operator=(const FComputeBufferReader& Other);
+	HORDE_API FComputeBufferReader& operator=(FComputeBufferReader&& Other) noexcept;
 
 	// Closes the handle to the underlying reader instance, resetting this instance back to empty
-	void Close();
+	HORDE_API void Close();
 
 	// Detaches this reader from the buffer, causing all pending and subsequent reads to return immediately.
-	void Detach();
+	HORDE_API void Detach();
 
 	// Test if the reader is valid
 	bool IsValid() const { return Detail != nullptr; }
 
 	// Test whether the buffer has finished being written to (ie. MarkComplete() has been called by the writer) and all data has been read from it.
-	bool IsComplete() const;
+	HORDE_API bool IsComplete() const;
 
 	// Move the read cursor forwards by the given number of bytes
-	void AdvanceReadPosition(size_t Size);
+	HORDE_API void AdvanceReadPosition(size_t Size);
 
 	// Gets the amount of data that is ready to be read from a contiguous block of memory.
-	size_t GetMaxReadSize() const;
+	HORDE_API size_t GetMaxReadSize() const;
 
 	// Reads data into the given buffer
-	size_t Read(void* Buffer, size_t MaxSize, int TimeoutMs = -1);
+	HORDE_API size_t Read(void* Buffer, size_t MaxSize, int TimeoutMs = -1);
 
 	// Waits until the given amount of data has been read, and returns a pointer to it. Returns nullptr if the timeout expires, or
 	// if the requested amount of data is not in a contiguous block of memory.
-	const unsigned char* WaitToRead(size_t MinSize, int TimeoutMs = -1);
+	HORDE_API const unsigned char* WaitToRead(size_t MinSize, int TimeoutMs = -1);
 
 private:
 	struct FReaderRef;
@@ -133,35 +134,38 @@ private:
 class FComputeBufferWriter
 {
 public:
-	FComputeBufferWriter();
-	FComputeBufferWriter(const FComputeBufferWriter& Other);
-	FComputeBufferWriter(FComputeBufferWriter&& Other) noexcept;
-	~FComputeBufferWriter();
+	HORDE_API FComputeBufferWriter();
+	HORDE_API FComputeBufferWriter(const FComputeBufferWriter& Other);
+	HORDE_API FComputeBufferWriter(FComputeBufferWriter&& Other) noexcept;
+	HORDE_API ~FComputeBufferWriter();
 
-	FComputeBufferWriter& operator=(const FComputeBufferWriter& Other);
-	FComputeBufferWriter& operator=(FComputeBufferWriter&& Other) noexcept;
+	HORDE_API FComputeBufferWriter& operator=(const FComputeBufferWriter& Other);
+	HORDE_API FComputeBufferWriter& operator=(FComputeBufferWriter&& Other) noexcept;
 
 	// Closes the handle to the underlying writer instance, resetting this instance back to empty
-	void Close();
+	HORDE_API void Close();
 
 	// Test if the writer is valid
 	bool IsValid() const { return Detail != nullptr; }
 
 	// Signal that we've finished writing to this buffer
-	void MarkComplete();
+	HORDE_API void MarkComplete();
 
 	// Move the write cursor forward by the given number of bytes
-	void AdvanceWritePosition(size_t Size);
+	HORDE_API void AdvanceWritePosition(size_t Size);
 
 	// Gets the length of the current write buffer.
-	size_t GetMaxWriteSize() const;
+	HORDE_API size_t GetMaxWriteSize() const;
+
+	// Get the max length of a chunk.
+	HORDE_API size_t GetChunkMaxLength() const;
 
 	// Writes data into the compute buffer
-	size_t Write(const void* Buffer, size_t MaxSize, int TimeoutMs = -1);
+	HORDE_API size_t Write(const void* Buffer, size_t MaxSize, int TimeoutMs = -1);
 
 	// Waits until a write buffer of the requested size is available, and returns a pointer to it. Returns nullptr if the 
 	// timeout expires before enough data has been flushed to return a buffer of the given length.
-	unsigned char* WaitToWrite(size_t MinSize, int TimeoutMs = -1);
+	HORDE_API unsigned char* WaitToWrite(size_t MinSize, int TimeoutMs = -1);
 
 private:
 	friend class FComputeBuffer;
diff --git a/Engine/Source/Developer/Horde/Public/ComputeChannel.h b/Engine/Source/Developer/Horde/Public/ComputeChannel.h
--- a/Engine/Source/Developer/Horde/Public/ComputeChannel.h
+++ b/Engine/Source/Developer/Horde/Public/ComputeChannel.h
@@ -2,6 +2,7 @@
 
 #pragma once
 
+#include "Horde.h"
 #include "ComputeBuffer.h"
 
 class FComputeSocket;
@@ -18,19 +19,19 @@ public:
 	// Writer for the channel
 	FComputeBufferWriter Writer;
 
-	FComputeChannel();
-	FComputeChannel(FComputeBufferReader InReader, FComputeBufferWriter InWriter);
-	~FComputeChannel();
+	HORDE_API FComputeChannel();
+	HORDE_API FComputeChannel(FComputeBufferReader InReader, FComputeBufferWriter InWriter);
+	HORDE_API ~FComputeChannel();
 
 	// Tests whether the channel is valid
-	bool IsValid() const;
+	HORDE_API bool IsValid() const;
 
 	// Sends bytes to the remote. 
-	size_t Send(const void* Data, size_t Size, int TimeoutMs = -1);
+	HORDE_API size_t Send(const void* Data, size_t Size, int TimeoutMs = -1);
 
 	// Reads as many bytes as are available from the socket.
-	size_t Recv(void* Data, size_t Size, int TimeoutMs = -1);
+	HORDE_API size_t Recv(void* Data, size_t Size, int TimeoutMs = -1);
 
 	// Indicate to the remote that no more data will be sent.
-	void MarkComplete();
+	HORDE_API void MarkComplete();
 };
diff --git a/Engine/Source/Developer/Horde/Public/ComputePlatform.h b/Engine/Source/Developer/Horde/Public/ComputePlatform.h
--- a/Engine/Source/Developer/Horde/Public/ComputePlatform.h
+++ b/Engine/Source/Developer/Horde/Public/ComputePlatform.h
@@ -6,6 +6,8 @@
 #include <stddef.h>
 #include <stdint.h>
 
+#include "Horde.h"
+
 // Defines for the current platform
 #ifdef _MSC_VER
 	#define UE_COMPUTE_PLATFORM_WINDOWS 1
@@ -39,19 +41,19 @@ public:
 	~FComputeEvent();
 
 	// Creates a new event with the given name
-	bool Create(const char* Name);
+	HORDE_API bool Create(const char* Name);
 
 	// Opens an existing event created elsewhere
-	bool OpenExisting(const char* Name);
+	HORDE_API bool OpenExisting(const char* Name);
 
 	// Close the event and release its resources
-	void Close();
+	HORDE_API void Close();
 
 	// Signal the event, releasing any waiters
-	void Signal();
+	HORDE_API void Signal();
 
 	// Wait for the event to be signalled or timeout. Pass -1 for timeoutMs to wait infinitely.
-	bool Wait(int timeoutMs);
+	HORDE_API bool Wait(int timeoutMs);
 
 private:
 	void* Handle;
@@ -63,20 +65,20 @@ private:
 class FComputeMemoryMappedFile
 {
 public:
-	FComputeMemoryMappedFile();
-	~FComputeMemoryMappedFile();
+	HORDE_API FComputeMemoryMappedFile();
+	HORDE_API ~FComputeMemoryMappedFile();
 
 	// Creates a new memory mapped file with the given capacity
-	bool Create(const char* Name, long long Capacity);
+	HORDE_API bool Create(const char* Name, long long Capacity);
 
 	// Opens an existing memory mapped file
-	bool OpenExisting(const char* Name);
+	HORDE_API bool OpenExisting(const char* Name);
 
 	// Close the memory mapped file handle
-	void Close();
+	HORDE_API void Close();
 
 	// Gets a pointer to the mapped memory
-	void* GetPointer() const;
+	HORDE_API void* GetPointer() const;
 
 private:
 	void* Handle;
@@ -95,71 +97,71 @@ struct FComputePlatform
 	//
 
 	// Signal that an assertion has failed
-	static void AssertFailed(const char* Expr, const char* File, int Line);
+	HORDE_API static void AssertFailed(const char* Expr, const char* File, int Line);
 	
 	// Reads an environment variable
-	static bool GetEnvironmentVariable(const char* Name, char* Buffer, size_t BufferLen);
+	HORDE_API static bool GetEnvironmentVariable(const char* Name, char* Buffer, size_t BufferLen);
 
 	// Creates a unique object name
-	static void CreateUniqueName(char* NameBuffer, size_t NameBufferLen);
+	HORDE_API static void CreateUniqueName(char* NameBuffer, size_t NameBufferLen);
 
 	//
 	// Math
 	//
 
 	// Find the log2 of the given value, returning 0 if the value is zero.
-	static unsigned int FloorLog2(unsigned int Value);
+	HORDE_API static unsigned int FloorLog2(unsigned int Value);
 
 	// Count the number of leading zeros in the given value.
-	static unsigned int CountLeadingZeros(unsigned int Value);
+	HORDE_API static unsigned int CountLeadingZeros(unsigned int Value);
 
 	//
 	// Strings
 	//
 
 	// Copy a string from one buffer to another, not exceeding the destination buffer size
-	static void Strcpy(char* Dest, size_t DestLen, const char* Source);
+	HORDE_API static void Strcpy(char* Dest, size_t DestLen, const char* Source);
 
 	// Perform a case-insensitive comparison of two strings
-	static int Stricmp(const char* A, const char* B);
+	HORDE_API static int Stricmp(const char* A, const char* B);
 
 	//
 	// Atomics
 	//
 
 	// Reads a 64-bit value from the given memory location 
-	static long long AtomicRead64(const volatile long long* Ptr);
+	HORDE_API static long long AtomicRead64(const volatile long long* Ptr);
 
 	// Writes a 64-bit value to the given memory location
-	static void AtomicWrite64(volatile long long* Ptr, long long Value);
+	HORDE_API static void AtomicWrite64(volatile long long* Ptr, long long Value);
 
 	// Increments an integer at the given memory location and returns the incremented value.
-	static long AtomicIncrement(volatile long* Ptr);
+	HORDE_API static long AtomicIncrement(volatile long* Ptr);
 
 	// Increments an integer at the given memory location and returns the incremented value.
-	static long long AtomicIncrement64(volatile long long* Ptr);
+	HORDE_API static long long AtomicIncrement64(volatile long long* Ptr);
 
 	// Decrements an integer at the given memory location and returns the decremented value.
-	static long AtomicDecrement(volatile long* Ptr);
+	HORDE_API static long AtomicDecrement(volatile long* Ptr);
 
 	// Adds a value to the given memory location, and returns the resulting value.
-	static long long AtomicAdd64(volatile long long* Ptr, long long Value);
+	HORDE_API static long long AtomicAdd64(volatile long long* Ptr, long long Value);
 
 	// ANDs a value with a given memory location, and returns the contents of the memory location BEFORE the operation.
-	static long AtomicAnd(volatile long* Ptr, long Value);
+	HORDE_API static long AtomicAnd(volatile long* Ptr, long Value);
 
 	// ANDs a value with a given memory location, and returns the contents of the memory location BEFORE the operation.
-	static long long AtomicAnd64(volatile long long* Ptr, long long Value);
+	HORDE_API static long long AtomicAnd64(volatile long long* Ptr, long long Value);
 
 	// ORs a value with a given memory location, and returns the contents of the memory location BEFORE the operation.
-	static long long AtomicOr64(volatile long long* Ptr, long long Value);
+	HORDE_API static long long AtomicOr64(volatile long long* Ptr, long long Value);
 
 	// Sets the contents of a pointer to the given value, and returns the original value in that memory location.
-	static long long AtomicExchange64(volatile long long* Ptr, long long Exchange);
+	HORDE_API static long long AtomicExchange64(volatile long long* Ptr, long long Exchange);
 
 	// Compares the contents of a memory location to a value, and exchanges it for another value if they are equal. Returns true if the operation succeeded.
-	static bool AtomicCompareExchange(volatile long* Ptr, long Exchange, long Comperand);
+	HORDE_API static bool AtomicCompareExchange(volatile long* Ptr, long Exchange, long Comperand);
 
 	// Compares the contents of a memory location to a value, and exchanges it for another value if they are equal. Returns true if the operation succeeded.
-	static bool AtomicCompareExchange64(volatile long long* Ptr, long long Exchange, long long Comperand);
+	HORDE_API static bool AtomicCompareExchange64(volatile long long* Ptr, long long Exchange, long long Comperand);
 };
diff --git a/Engine/Source/Developer/Horde/Public/ComputeSocket.h b/Engine/Source/Developer/Horde/Public/ComputeSocket.h
--- a/Engine/Source/Developer/Horde/Public/ComputeSocket.h
+++ b/Engine/Source/Developer/Horde/Public/ComputeSocket.h
@@ -14,23 +14,26 @@
 class FComputeSocket
 {
 public:
-	FComputeSocket();
-	virtual ~FComputeSocket();
+	HORDE_API FComputeSocket();
+	HORDE_API virtual ~FComputeSocket();
 
-	FComputeSocket(const FComputeSocket&) = delete;
-	FComputeSocket& operator=(const FComputeSocket&) = delete;
+	HORDE_API FComputeSocket(const FComputeSocket&) = delete;
+	HORDE_API FComputeSocket& operator=(const FComputeSocket&) = delete;
+
+	// Begins communication with the agent
+	HORDE_API virtual void StartCommunication() = 0;
 
 	// Attaches a new buffer for receiving data
-	virtual void AttachRecvBuffer(int ChannelId, FComputeBuffer RecvBuffer) = 0;
+	HORDE_API virtual void AttachRecvBuffer(int ChannelId, FComputeBuffer RecvBuffer) = 0;
 
 	// Attaches a new buffer for sending data */
-	virtual void AttachSendBuffer(int ChannelId, FComputeBuffer SendBuffer) = 0;
+	HORDE_API virtual void AttachSendBuffer(int ChannelId, FComputeBuffer SendBuffer) = 0;
 
 	// Attaches a channel to this socket
-	FComputeChannel CreateChannel(int ChannelId);
+	HORDE_API std::shared_ptr<FComputeChannel> CreateChannel(int ChannelId);
 
 	// Attaches a channel to this socket
-	FComputeChannel CreateChannel(int ChannelId, FComputeBuffer RecvBuffer, FComputeBuffer SendBuffer);
+	HORDE_API std::shared_ptr<FComputeChannel> CreateChannel(int ChannelId, FComputeBuffer RecvBuffer, FComputeBuffer SendBuffer);
 };
 
 //
@@ -42,26 +45,29 @@ class FWorkerComputeSocket final : public FComputeSocket
 public:
 	static const char* const IpcEnvVar;
 
-	FWorkerComputeSocket();
-	~FWorkerComputeSocket();
+	HORDE_API FWorkerComputeSocket();
+	HORDE_API ~FWorkerComputeSocket();
 
 	// Opens a connection to the agent process using a command buffer read from an environment variable (EnvVarName)
-	bool Open();
+	HORDE_API bool Open();
 
 	// Opens a connection to the agent process using a specific command buffer name
-	bool Open(const char* CommandBufferName);
+	HORDE_API bool Open(const char* CommandBufferName);
 
 	// Close the current connection
-	void Close();
+	HORDE_API void Close();
+
+	// Begin communication with agent.
+	HORDE_API virtual void StartCommunication() override;
 
 	// Attaches a new buffer for receiving data
-	virtual void AttachRecvBuffer(int ChannelId, FComputeBuffer RecvBuffer) override;
+	HORDE_API virtual void AttachRecvBuffer(int ChannelId, FComputeBuffer RecvBuffer) override;
 
 	// Attaches a new buffer for sending data
-	virtual void AttachSendBuffer(int ChannelId, FComputeBuffer SendBuffer) override;
+	HORDE_API virtual void AttachSendBuffer(int ChannelId, FComputeBuffer SendBuffer) override;
 
 	// Reads and handles a command from the command buffer
-	static void RunServer(FComputeBufferReader& CommandBufferReader, FComputeSocket& Socket);
+	HORDE_API static void RunServer(FComputeBufferReader& CommandBufferReader, FComputeSocket& Socket);
 
 private:
 	enum class EMessageType;
@@ -90,5 +96,5 @@ enum class EComputeSocketEndpoint
 	Remote
 };
 
-// Creates a socket using a custom transport
-std::unique_ptr<FComputeSocket> CreateComputeSocket(std::unique_ptr<FComputeTransport> Transport, EComputeSocketEndpoint Endpoint);
+// Creates a socket using a custom transport. Also returns the default channel (channel 0)
+HORDE_API std::unique_ptr<FComputeSocket> CreateComputeSocket(std::unique_ptr<FComputeTransport> Transport, EComputeSocketEndpoint Endpoint);
diff --git a/Engine/Source/Developer/Horde/Public/ComputeTransport.h b/Engine/Source/Developer/Horde/Public/ComputeTransport.h
--- a/Engine/Source/Developer/Horde/Public/ComputeTransport.h
+++ b/Engine/Source/Developer/Horde/Public/ComputeTransport.h
@@ -2,6 +2,7 @@
 
 #pragma once
 
+#include "Horde.h"
 #include "ComputeBuffer.h"
 
 //
@@ -10,7 +11,7 @@
 class FComputeTransport
 {
 public:
-	virtual ~FComputeTransport();
+	HORDE_API virtual ~FComputeTransport();
 	
 	// Sends data to the remote
 	virtual size_t Send(const void* Data, size_t Size) = 0;
@@ -25,10 +26,10 @@ public:
 	virtual void Close() = 0;
 
 	// Sends data to the remote, blocking until the entire message has been sent.
-	bool SendMessage(const void* Data, size_t Size);
+	HORDE_API bool SendMessage(const void* Data, size_t Size);
 
 	// Receives a fixed length block of data from the remote, blocking until the entire length has been received.
-	bool RecvMessage(void* Data, size_t Size);
+	HORDE_API bool RecvMessage(void* Data, size_t Size);
 };
 
 //
@@ -37,7 +38,7 @@ public:
 class FBufferTransport final : public FComputeTransport
 {
 public:
-	FBufferTransport(FComputeBufferWriter InSendBufferWriter, FComputeBufferReader InRecvBufferReader);
+	HORDE_API FBufferTransport(FComputeBufferWriter InSendBufferWriter, FComputeBufferReader InRecvBufferReader);
 
 protected:
 	virtual size_t Send(const void* Data, size_t Size) override final;
diff --git a/Engine/Source/Developer/Horde/Public/Horde.h b/Engine/Source/Developer/Horde/Public/Horde.h
--- a/Engine/Source/Developer/Horde/Public/Horde.h
+++ b/Engine/Source/Developer/Horde/Public/Horde.h
@@ -2,11 +2,14 @@
 
 #pragma once
 
-#include "CoreMinimal.h"
-#include "Containers/UnrealString.h"
+#if !defined(HORDE_API)
+#define HORDE_API
+#endif
 
 #if WITH_EDITOR
 
+#include "Containers/UnrealString.h"
+
 struct FHorde
 {
 	HORDE_API static FString GetTemplateName();
diff --git a/Engine/Source/Programs/Horde/Samples/RemoteWorkerCpp/Main.cpp b/Engine/Source/Programs/Horde/Samples/RemoteWorkerCpp/Main.cpp
--- a/Engine/Source/Programs/Horde/Samples/RemoteWorkerCpp/Main.cpp
+++ b/Engine/Source/Programs/Horde/Samples/RemoteWorkerCpp/Main.cpp
@@ -26,8 +26,8 @@ int main(int argc, const char* argv[])
 		return 1;
 	}
 
-	FComputeChannel Channel = Socket.CreateChannel(ChannelId);
-	if(!Channel.IsValid())
+	std::shared_ptr<FComputeChannel> Channel = Socket.CreateChannel(ChannelId);
+	if(!Channel->IsValid())
 	{
 		std::cout << "Unable to create channel to initiator" << std::endl;
 		return 1;
@@ -39,11 +39,11 @@ int main(int argc, const char* argv[])
 	char Buffer[4];
 
 	Buffer[0] = 0;
-	Channel.Send(Buffer, 1); // Send a dummy one-byte message to let the remote know we're listening
+	Channel->Send(Buffer, 1); // Send a dummy one-byte message to let the remote know we're listening
 
 	for (;;)
 	{
-		size_t RecvLength = Channel.Recv(Buffer + Length, sizeof(Buffer) - Length);
+		size_t RecvLength = Channel->Recv(Buffer + Length, sizeof(Buffer) - Length);
 		if (RecvLength == 0)
 		{
 			return 0;
diff --git a/Engine/Source/Programs/Horde/Samples/RemoteWorkerCpp/Tests.cpp b/Engine/Source/Programs/Horde/Samples/RemoteWorkerCpp/Tests.cpp
--- a/Engine/Source/Programs/Horde/Samples/RemoteWorkerCpp/Tests.cpp
+++ b/Engine/Source/Programs/Horde/Samples/RemoteWorkerCpp/Tests.cpp
@@ -41,12 +41,12 @@ void ComputeBufferTest()
 	UE_COMPUTE_ASSERT(Reader.IsComplete());
 }
 
-template<size_t TestDataSize> void CheckChannelSendRecv(FComputeChannel& SendChannel, FComputeChannel& RecvChannel, const char(&TestData)[TestDataSize])
+template<size_t TestDataSize> void CheckChannelSendRecv(std::shared_ptr<FComputeChannel>& SendChannel, std::shared_ptr<FComputeChannel>& RecvChannel, const char(&TestData)[TestDataSize])
 {
-	SendChannel.Send(TestData, TestDataSize);
+	SendChannel->Send(TestData, TestDataSize);
 
 	char RecvTestData[TestDataSize];
-	UE_COMPUTE_VERIFY(RecvChannel.Recv(RecvTestData, TestDataSize) == TestDataSize);
+	UE_COMPUTE_VERIFY(RecvChannel->Recv(RecvTestData, TestDataSize) == TestDataSize);
 
 	UE_COMPUTE_VERIFY(memcmp(TestData, RecvTestData, TestDataSize) == 0);
 }
@@ -62,13 +62,13 @@ void ComputeSocketTest()
 
 	// Client transport
 	std::unique_ptr<FComputeSocket> ClientSocket = CreateComputeSocket(std::make_unique<FBufferTransport>(ClientToServerBuffer.CreateWriter(), ServerToClientBuffer.CreateReader()), EComputeSocketEndpoint::Local);
-	FComputeChannel ClientChannel1 = ClientSocket->CreateChannel(1);
-	FComputeChannel ClientChannel2 = ClientSocket->CreateChannel(2);
+	std::shared_ptr<FComputeChannel> ClientChannel1 = ClientSocket->CreateChannel(1);
+	std::shared_ptr<FComputeChannel> ClientChannel2 = ClientSocket->CreateChannel(2);
 
 	// Server socket
 	std::unique_ptr<FComputeSocket> ServerSocket = CreateComputeSocket(std::make_unique<FBufferTransport>(ServerToClientBuffer.CreateWriter(), ClientToServerBuffer.CreateReader()), EComputeSocketEndpoint::Remote);
-	FComputeChannel ServerChannel1 = ServerSocket->CreateChannel(1);
-	FComputeChannel ServerChannel2 = ServerSocket->CreateChannel(2);
+	std::shared_ptr<FComputeChannel> ServerChannel1 = ServerSocket->CreateChannel(1);
+	std::shared_ptr<FComputeChannel> ServerChannel2 = ServerSocket->CreateChannel(2);
 
 	// Close the original buffers now that the sockets are set up
 	ClientToServerBuffer.Close();
