Horde: Fix Clang compile error (cannot apply dllexport to deleted function).

diff --git a/Engine/Source/Developer/Horde/Public/ComputeSocket.h b/Engine/Source/Developer/Horde/Public/ComputeSocket.h
--- a/Engine/Source/Developer/Horde/Public/ComputeSocket.h
+++ b/Engine/Source/Developer/Horde/Public/ComputeSocket.h
@@ -17,8 +17,8 @@ public:
 	HORDE_API FComputeSocket();
 	HORDE_API virtual ~FComputeSocket();
 
-	HORDE_API FComputeSocket(const FComputeSocket&) = delete;
-	HORDE_API FComputeSocket& operator=(const FComputeSocket&) = delete;
+	FComputeSocket(const FComputeSocket&) = delete;
+	FComputeSocket& operator=(const FComputeSocket&) = delete;
 
 	// Begins communication with the agent
 	HORDE_API virtual void StartCommunication() = 0;
