Add null pointer assertion for Horde agent message header.

#rnx

diff --git a/Engine/Source/Developer/Horde/Private/AgentMessage.cpp b/Engine/Source/Developer/Horde/Private/AgentMessage.cpp
--- a/Engine/Source/Developer/Horde/Private/AgentMessage.cpp
+++ b/Engine/Source/Developer/Horde/Private/AgentMessage.cpp
@@ -162,6 +162,7 @@ EAgentMessageType FAgentMessageChannel::ReadResponse()
 	}
 
 	const unsigned char* Header = ChannelBuffers->Reader.WaitToRead(MessageHeaderLength);
+	check(Header != nullptr);
 	unsigned int Length = *((unsigned int*)(Header + 1));
 
 	Header = ChannelBuffers->Reader.WaitToRead(MessageHeaderLength + Length);
