Build fix on Mac/Linux.

diff --git a/Engine/Source/Developer/Horde/Private/Compute/ComputePlatform.cpp b/Engine/Source/Developer/Horde/Private/Compute/ComputePlatform.cpp
--- a/Engine/Source/Developer/Horde/Private/Compute/ComputePlatform.cpp
+++ b/Engine/Source/Developer/Horde/Private/Compute/ComputePlatform.cpp
@@ -373,7 +373,7 @@ void FComputePlatform::CreateUniqueName(char* NameBuffer, size_t NameBufferLen)
 #else
 	struct timespec ts;
 	clock_gettime(CLOCK_MONOTONIC, &ts);
-	snprintf(NameBuffer, NameBufferLen, "/UEC_%u%zu%zu_%lu", getpid(), (size_t)ts.tv_sec, (size_t)ts.tv_nsec, AtomicIncrement(&Counter));
+	snprintf(NameBuffer, NameBufferLen, "/UEC_%u%zu%zu_%lu", getpid(), (size_t)ts.tv_sec, (size_t)ts.tv_nsec, FPlatformAtomics::InterlockedIncrement(&Counter));
 #endif
 }
 
