Horde: Move ComputePlatform.h out of public API.

diff --git a/Engine/Source/Developer/Horde/Private/Compute/AgentMessage.cpp b/Engine/Source/Developer/Horde/Private/Compute/AgentMessage.cpp
--- a/Engine/Source/Developer/Horde/Private/Compute/AgentMessage.cpp
+++ b/Engine/Source/Developer/Horde/Private/Compute/AgentMessage.cpp
@@ -2,6 +2,7 @@
 
 #include "Compute/AgentMessage.h"
 #include "Compute/ComputePlatform.h"
+#include "../HordePlatform.h"
 
 //// FAgentMessageChannel ////
 
@@ -258,7 +259,7 @@ size_t FAgentMessageChannel::MeasureUnsignedVarInt(size_t Value)
 	}
 	else
 	{
-		return (FComputePlatform::FloorLog2((unsigned int)Value) / 7) + 1;
+		return (FHordePlatform::FloorLog2((unsigned int)Value) / 7) + 1;
 	}
 }
 
@@ -283,7 +284,7 @@ size_t FAgentMessageChannel::ReadUnsignedVarInt(const unsigned char** Pos)
 	const unsigned char* Data = *Pos;
 
 	unsigned char FirstByte = Data[0];
-	size_t NumBytes = FComputePlatform::CountLeadingZeros(0xFF & (~(unsigned int)FirstByte)) + 1 - 24; // Note byte -> int conversion here, hence ignoring subtracting 24 bits
+	size_t NumBytes = FHordePlatform::CountLeadingZeros(0xFF & (~(unsigned int)FirstByte)) + 1 - 24; // Note byte -> int conversion here, hence ignoring subtracting 24 bits
 
 	size_t value = (size_t)(FirstByte & (0xff >> NumBytes));
 	for (size_t Idx = 1; Idx < NumBytes; Idx++)
diff --git a/Engine/Source/Developer/Horde/Private/Compute/ComputeBuffer.cpp b/Engine/Source/Developer/Horde/Private/Compute/ComputeBuffer.cpp
--- a/Engine/Source/Developer/Horde/Private/Compute/ComputeBuffer.cpp
+++ b/Engine/Source/Developer/Horde/Private/Compute/ComputeBuffer.cpp
@@ -2,6 +2,7 @@
 
 #include "Compute/ComputeBuffer.h"
 #include "Compute/ComputePlatform.h"
+#include "../HordePlatform.h"
 #include <assert.h>
 #include <iostream>
 #include <assert.h>
@@ -313,7 +314,7 @@ struct FComputeBufferDetail
 		static_assert(sizeof(FChunkStatePtr) == sizeof(long long), "Incorrect size of FChunkStatePtr; check union is declared correctly.");
 		static_assert(sizeof(FWriterStatePtr) == sizeof(long long), "Incorrect size of FWriterStatePtr; check union is declared correctly.");
 
-		FComputePlatform::Strcpy(Name, FComputeBuffer::MaxNameLength, InName);
+		FCStringAnsi::Strcpy(Name, FComputeBuffer::MaxNameLength, InName);
 	}
 
 	~FComputeBufferDetail()
@@ -343,7 +344,7 @@ struct FComputeBufferDetail
 		char BaseNameBuffer[FComputeBuffer::MaxNameLength];
 		if (Name == nullptr)
 		{
-			FComputePlatform::CreateUniqueName(BaseNameBuffer, FComputeBuffer::MaxNameLength);
+			FHordePlatform::CreateUniqueName(BaseNameBuffer, FComputeBuffer::MaxNameLength);
 			Name = BaseNameBuffer;
 		}
 
diff --git a/Engine/Source/Developer/Horde/Private/Compute/ComputePlatform.cpp b/Engine/Source/Developer/Horde/Private/Compute/ComputePlatform.cpp
--- a/Engine/Source/Developer/Horde/Private/Compute/ComputePlatform.cpp
+++ b/Engine/Source/Developer/Horde/Private/Compute/ComputePlatform.cpp
@@ -333,72 +333,3 @@ void* FComputeMemoryMappedFile::GetPointer() const
 	return (unsigned char*)Pointer + 16;
 #endif
 }
-
-/////////////////////////////////////////////////// 
-
-void FComputePlatform::AssertFailed(const char* Expr, const char* File, int Line)
-{
-	std::cerr << "Assertion failed: '" << Expr << std::endl;
-	std::cerr << "  at " << File << "(" << Line << ")" << std::endl;
-#if UE_COMPUTE_PLATFORM_WINDOWS
-	DebugBreak();
-#endif
-	exit(1);
-}
-
-bool FComputePlatform::GetEnvironmentVariable(const char* Name, char* Buffer, size_t BufferLen)
-{
-#if UE_COMPUTE_PLATFORM_WINDOWS
-	int Length = GetEnvironmentVariableA(Name, Buffer, (DWORD)BufferLen);
-	return Length > 0 && Length < BufferLen;
-#else
-	char* Value = getenv(Name);
-	if(Value != nullptr)
-	{
-		FComputePlatform::Strcpy(Buffer, BufferLen, Value);
-		return true;
-	}
-	return false;
-#endif
-}
-
-void FComputePlatform::CreateUniqueName(char* NameBuffer, size_t NameBufferLen)
-{
-	static int32 Counter = 0;
-
-#if UE_COMPUTE_PLATFORM_WINDOWS
-	DWORD Pid = GetCurrentProcessId();
-	ULONGLONG TickCount = GetTickCount64();
-	snprintf(NameBuffer, NameBufferLen, "Local\\COMPUTE_%lu_%llu_%lu", Pid, TickCount, (unsigned long)FPlatformAtomics::InterlockedIncrement(&Counter));
-#else
-	struct timespec ts;
-	clock_gettime(CLOCK_MONOTONIC, &ts);
-	snprintf(NameBuffer, NameBufferLen, "/UEC_%u%zu%zu_%d", getpid(), (size_t)ts.tv_sec, (size_t)ts.tv_nsec, FPlatformAtomics::InterlockedIncrement(&Counter));
-#endif
-}
-
-unsigned int FComputePlatform::FloorLog2(unsigned int Value)
-{
-	return std::max<unsigned int>(0, 31 - std::countl_zero(Value));
-}
-
-unsigned int FComputePlatform::CountLeadingZeros(unsigned int Value)
-{
-	return std::countl_zero(Value);
-}
-
-void FComputePlatform::Strcpy(char* Dest, size_t DestLen, const char* Source)
-{
-	size_t Length = std::min(strlen(Source), DestLen - 1);
-	memcpy(Dest, Source, Length);
-	Dest[Length] = 0;
-}
-
-int FComputePlatform::Stricmp(const char* A, const char* B)
-{
-#if UE_COMPUTE_PLATFORM_WINDOWS
-	return _stricmp(A, B);
-#else
-	return strcasecmp(A, B);
-#endif
-}
diff --git a/Engine/Source/Developer/Horde/Private/Compute/ComputePlatform.h b/Engine/Source/Developer/Horde/Private/Compute/ComputePlatform.h
new file mode 100644
--- /dev/null
+++ b/Engine/Source/Developer/Horde/Private/Compute/ComputePlatform.h
@@ -0,0 +1,64 @@
+// Copyright Epic Games, Inc. All Rights Reserved.
+
+#pragma once
+
+#include "CoreMinimal.h"
+#include <stdio.h>
+#include <stddef.h>
+#include <stdint.h>
+
+//
+// IPC manually-reset event object
+//
+class FComputeEvent
+{
+public:
+	FComputeEvent();
+	~FComputeEvent();
+
+	// Creates a new event with the given name
+	HORDE_API bool Create(const char* Name);
+
+	// Opens an existing event created elsewhere
+	HORDE_API bool OpenExisting(const char* Name);
+
+	// Close the event and release its resources
+	HORDE_API void Close();
+
+	// Signal the event, releasing any waiters
+	HORDE_API void Signal();
+
+	// Wait for the event to be signalled or timeout. Pass -1 for timeoutMs to wait infinitely.
+	HORDE_API bool Wait(int timeoutMs);
+
+private:
+	void* Handle;
+};
+
+//
+// IPC memory mapped file
+//
+class FComputeMemoryMappedFile
+{
+public:
+	HORDE_API FComputeMemoryMappedFile();
+	HORDE_API ~FComputeMemoryMappedFile();
+
+	// Creates a new memory mapped file with the given capacity
+	HORDE_API bool Create(const char* Name, long long Capacity);
+
+	// Opens an existing memory mapped file
+	HORDE_API bool OpenExisting(const char* Name);
+
+	// Close the memory mapped file handle
+	HORDE_API void Close();
+
+	// Gets a pointer to the mapped memory
+	HORDE_API void* GetPointer() const;
+
+private:
+	void* Handle;
+	void* Pointer;
+	long long MappedSize;
+	char* OwnerName;
+};
diff --git a/Engine/Source/Developer/Horde/Private/Compute/ComputeSocket.cpp b/Engine/Source/Developer/Horde/Private/Compute/ComputeSocket.cpp
--- a/Engine/Source/Developer/Horde/Private/Compute/ComputeSocket.cpp
+++ b/Engine/Source/Developer/Horde/Private/Compute/ComputeSocket.cpp
@@ -10,6 +10,7 @@
 #include <thread>
 #include <mutex>
 #include <chrono>
+#include "../HordePlatform.h"
 
 FComputeSocket::FComputeSocket()
 {
@@ -72,7 +73,7 @@ void FWorkerComputeSocket::StartCommunication()
 bool FWorkerComputeSocket::Open()
 {
 	char EnvVar[FComputeBuffer::MaxNameLength];
-	if (!FComputePlatform::GetEnvironmentVariable(IpcEnvVar, EnvVar, sizeof(EnvVar) / sizeof(EnvVar[0])))
+	if (!FHordePlatform::GetEnvironmentVariable(IpcEnvVar, EnvVar, sizeof(EnvVar) / sizeof(EnvVar[0])))
 	{
 		return false;
 	}
@@ -182,7 +183,7 @@ void FWorkerComputeSocket::RunServer(FComputeBufferReader& CommandBufferReader,
 
 size_t FWorkerComputeSocket::ReadVarUInt(const unsigned char* Pos, unsigned int* OutValue)
 {
-	size_t ByteCount = FComputePlatform::CountLeadingZeros((unsigned char)(~*static_cast<const unsigned char*>(Pos))) - 23;
+	size_t ByteCount = FHordePlatform::CountLeadingZeros((unsigned char)(~*static_cast<const unsigned char*>(Pos))) - 23;
 
 	unsigned int Value = *Pos++ & (unsigned char)(0xff >> ByteCount);
 	switch (ByteCount - 1)
@@ -208,7 +209,7 @@ size_t FWorkerComputeSocket::ReadString(const unsigned char* Pos, char* OutText,
 	unsigned int TextLen;
 
 	size_t Len = ReadVarUInt(Pos, &TextLen);
-	FComputePlatform::Strcpy(OutText, OutTextMaxLen, (const char*)Pos + Len);
+	FCStringAnsi::Strcpy(OutText, OutTextMaxLen, (const char*)Pos + Len);
 
 	return Len + TextLen;
 }
@@ -217,7 +218,7 @@ size_t FWorkerComputeSocket::WriteVarUInt(unsigned char* Pos, unsigned int Value
 {
 	// Use BSR to return the log2 of the integer
 	// return 0 if value is 0
-	unsigned int ByteCount = (unsigned int)(int(FComputePlatform::FloorLog2(Value)) / 7 + 1);
+	unsigned int ByteCount = (unsigned int)(int(FHordePlatform::FloorLog2(Value)) / 7 + 1);
 
 	unsigned char* OutBytes = Pos + ByteCount - 1;
 	switch (ByteCount - 1)
diff --git a/Engine/Source/Developer/Horde/Private/HordePlatform.cpp b/Engine/Source/Developer/Horde/Private/HordePlatform.cpp
new file mode 100644
--- /dev/null
+++ b/Engine/Source/Developer/Horde/Private/HordePlatform.cpp
@@ -0,0 +1,124 @@
+// Copyright Epic Games, Inc. All Rights Reserved.
+
+#include "HordePlatform.h"
+#include <assert.h>
+#include <wchar.h>
+#include <bit>
+#include <algorithm>
+#include <iostream>
+
+// Defines for the current platform
+#ifdef _MSC_VER
+#define UE_HORDE_PLATFORM_WINDOWS 1
+#define UE_HORDE_PLATFORM_MAC 0
+#define UE_HORDE_PLATFORM_LINUX 0
+#elif defined(__APPLE__)
+#define UE_HORDE_PLATFORM_WINDOWS 0
+#define UE_HORDE_PLATFORM_MAC 1
+#define UE_HORDE_PLATFORM_LINUX 0
+#else
+#define UE_HORDE_PLATFORM_WINDOWS 0
+#define UE_HORDE_PLATFORM_MAC 0
+#define UE_HORDE_PLATFORM_LINUX 1
+#endif
+
+#if UE_HORDE_PLATFORM_WINDOWS
+#include <Windows.h>
+#undef min
+#undef max
+#undef GetEnvironmentVariable
+#undef SendMessage
+#undef InterlockedIncrement
+#else
+#include <semaphore.h>
+#include <unistd.h>
+#include <atomic>
+#include <sys/mman.h>
+#include <sys/stat.h>
+#include <time.h>
+#include <fcntl.h>
+#endif
+
+void FHordePlatform::NotImplemented()
+{
+	throw std::string("Not Implemented");
+}
+
+void FHordePlatform::NotSupported(const char* Message)
+{
+	throw std::string(Message);
+}
+
+bool FHordePlatform::GetEnvironmentVariable(const char* Name, char* Buffer, size_t BufferLen)
+{
+#if PLATFORM_WINDOWS
+	int Length = GetEnvironmentVariableA(Name, Buffer, (DWORD)BufferLen);
+	return Length > 0 && Length < BufferLen;
+#else
+	char* Value = getenv(Name);
+	if (Value != nullptr)
+	{
+		FCStringAnsi::Strcpy(Buffer, BufferLen, Value);
+		return true;
+	}
+	return false;
+#endif
+}
+
+void FHordePlatform::CreateUniqueIdentifier(char* NameBuffer, size_t NameBufferLen)
+{
+	static int32 Counter = 0;
+
+#if UE_HORDE_PLATFORM_WINDOWS
+	DWORD Pid = GetCurrentProcessId();
+	ULONGLONG TickCount = GetTickCount64();
+	snprintf(NameBuffer, NameBufferLen, "%lu_%llu_%lu", Pid, TickCount, (unsigned long)FPlatformAtomics::InterlockedIncrement(&Counter));
+#else
+	struct timespec ts;
+	clock_gettime(CLOCK_MONOTONIC, &ts);
+	snprintf(NameBuffer, NameBufferLen, "%u%zu%zu_%lu", getpid(), (size_t)ts.tv_sec, (size_t)ts.tv_nsec, AtomicIncrement(&Counter));
+#endif
+}
+
+void FHordePlatform::CreateUniqueName(char* NameBuffer, size_t NameBufferLen)
+{
+#if UE_HORDE_PLATFORM_WINDOWS
+	const char Prefix[] = "Local\\COMPUTE_";
+#else
+	const char Prefix[] = "/UEC_";
+#endif
+
+	size_t PrefixLen = (sizeof(Prefix) / sizeof(Prefix[0])) - 1;
+	check(NameBufferLen > PrefixLen);
+
+	memcpy(NameBuffer, Prefix, PrefixLen * sizeof(char));
+	CreateUniqueIdentifier(NameBuffer + PrefixLen, NameBufferLen - PrefixLen);
+}
+
+unsigned int FHordePlatform::FloorLog2(unsigned int Value)
+{
+	return std::max<unsigned int>(0, 31 - std::countl_zero(Value));
+}
+
+unsigned int FHordePlatform::CountLeadingZeros(unsigned int Value)
+{
+	return std::countl_zero(Value);
+}
+
+bool FHordePlatform::TryParseSizeT(const char* Source, size_t SourceLen, size_t& OutValue, size_t& OutNumBytes)
+{
+	OutValue = 0;
+	OutNumBytes = 0;
+
+	for (; OutNumBytes < SourceLen; OutNumBytes++)
+	{
+		char Character = Source[OutNumBytes];
+		if (Character < '0' || Character > '9')
+		{
+			break;
+		}
+		OutValue = (OutValue * 10) + (Character - '0');
+	}
+
+	return OutNumBytes > 0;
+}
diff --git a/Engine/Source/Developer/Horde/Private/HordePlatform.h b/Engine/Source/Developer/Horde/Private/HordePlatform.h
new file mode 100644
--- /dev/null
+++ b/Engine/Source/Developer/Horde/Private/HordePlatform.h
@@ -0,0 +1,47 @@
+// Copyright Epic Games, Inc. All Rights Reserved.
+
+#pragma once
+
+#include "CoreMinimal.h"
+
+//
+// Platform-specific utility functions
+//
+struct FHordePlatform
+{
+	//
+	// General
+	//
+
+	// Indicate that the current codepath is not implemented
+	[[noreturn]] HORDE_API static void NotImplemented();
+
+	// Indicate that the current codepath is not supported
+	[[noreturn]] HORDE_API static void NotSupported(const char* Message);
+
+	// Reads an environment variable
+	HORDE_API static bool GetEnvironmentVariable(const char* Name, char* Buffer, size_t BufferLen);
+
+	// Creates a unique identifier
+	HORDE_API static void CreateUniqueIdentifier(char* NameBuffer, size_t NameBufferLen);
+
+	// Creates a unique object name
+	HORDE_API static void CreateUniqueName(char* NameBuffer, size_t NameBufferLen);
+
+	//
+	// Math
+	//
+
+	// Find the log2 of the given value, returning 0 if the value is zero.
+	HORDE_API static unsigned int FloorLog2(unsigned int Value);
+
+	// Count the number of leading zeros in the given value.
+	HORDE_API static unsigned int CountLeadingZeros(unsigned int Value);
+
+	//
+	// Strings
+	//
+
+	// Parse an integer from a string, returning the number of bytes read
+	HORDE_API static bool TryParseSizeT(const char* Source, size_t SourceLen, size_t& OutValue, size_t& OutNumBytes);
+};
diff --git a/Engine/Source/Developer/Horde/Public/Compute/ComputePlatform.h b/Engine/Source/Developer/Horde/Public/Compute/ComputePlatform.h
deleted file mode 100644
diff --git a/Engine/Source/Programs/Horde/Samples/Engine/Core/Public/CoreMinimal.h b/Engine/Source/Programs/Horde/Samples/Engine/Core/Public/CoreMinimal.h
--- a/Engine/Source/Programs/Horde/Samples/Engine/Core/Public/CoreMinimal.h
+++ b/Engine/Source/Programs/Horde/Samples/Engine/Core/Public/CoreMinimal.h
@@ -158,6 +158,7 @@ using WIDECHAR = char16_t;
 #define UE_LOG(...)
 #define check(x)						do { if (!(x)) __debugbreak(); } while (0)
 #define checkf(x, ...)					check(x)
+#define verify(x)						do { if (!(x)) __debugbreak(); } while (0)
 
 #if !defined(UE_BUILD_SHIPPING)
 #	define UE_BUILD_SHIPPING			1
@@ -489,6 +490,35 @@ struct FPlatformAtomics
 	}
 };
 
+template<typename CharType>
+struct TCString
+{
+	static FORCEINLINE CharType* Strcpy(CharType* Dest, SIZE_T DestCount, const CharType* Src);
+
+	template<SIZE_T DestCount>
+	static FORCEINLINE CharType* Strcpy(CharType(&Dest)[DestCount], const CharType* Src)
+	{
+		return Strcpy(Dest, DestCount, Src);
+	}
+
+	static int32 Stricmp(const CharType* String1, const CharType* String2);
+};
+
+typedef TCString<char> FCStringAnsi;
+
+template<> 
+inline char* TCString<char>::Strcpy(char* Dest, size_t DestCount, const char* Src)
+{
+	strcpy_s(Dest, DestCount, Src);
+	return Dest;
+}
+
+template<>
+inline int32 TCString<char>::Stricmp(const char* String1, const char* String2)
+{
+	return _stricmp(String1, String2);
+}
+
 namespace UECompat
 {
 
@@ -538,7 +568,6 @@ namespace UECompat
 
 namespace Algo = UECompat;
 namespace FMath = UECompat;
-namespace FCStringAnsi = UECompat;
 namespace FMemory = UECompat;
 
 #endif // TRACE_UE_COMPAT_LAYER
