Build fix on Mac/Linux editor.

diff --git a/Engine/Source/Developer/Horde/Private/HordePlatform.cpp b/Engine/Source/Developer/Horde/Private/HordePlatform.cpp
--- a/Engine/Source/Developer/Horde/Private/HordePlatform.cpp
+++ b/Engine/Source/Developer/Horde/Private/HordePlatform.cpp
@@ -76,7 +76,7 @@ void FHordePlatform::CreateUniqueIdentifier(char* NameBuffer, size_t NameBufferL
 #else
 	struct timespec ts;
 	clock_gettime(CLOCK_MONOTONIC, &ts);
-	snprintf(NameBuffer, NameBufferLen, "%u%zu%zu_%lu", getpid(), (size_t)ts.tv_sec, (size_t)ts.tv_nsec, AtomicIncrement(&Counter));
+	snprintf(NameBuffer, NameBufferLen, "%u%zu%zu_%d", getpid(), (size_t)ts.tv_sec, (size_t)ts.tv_nsec, FPlatformAtomics::InterlockedIncrement(&Counter));
 #endif
 }
 
