Initialize MaxRequestSize in Horde agent message channel when a message is created.

#rnx

diff --git a/Engine/Source/Developer/Horde/Private/Compute/AgentMessage.cpp b/Engine/Source/Developer/Horde/Private/Compute/AgentMessage.cpp
--- a/Engine/Source/Developer/Horde/Private/Compute/AgentMessage.cpp
+++ b/Engine/Source/Developer/Horde/Private/Compute/AgentMessage.cpp
@@ -211,6 +211,7 @@ void FAgentMessageChannel::CreateMessage(EAgentMessageType Type, size_t MaxLengt
 {
 	RequestData = ChannelBuffers->Writer.WaitToWrite(MessageHeaderLength + MaxLength);
 	RequestData[0] = (unsigned char)Type;
+	MaxRequestSize = MaxLength;
 }
 
 void FAgentMessageChannel::FlushMessage()
