Horde: Initial port of EpicGames.Horde storage library to C++.

Still plenty of missing features:
* Content-defined chunking
* Caching
* Hashing and deduplication of nodes through BlobWriter
* Writing of aliases and refs through BlobWriter
* HTTP client

diff --git a/Engine/Source/Developer/Horde/Private/Compute/ComputeChannel.cpp b/Engine/Source/Developer/Horde/Private/Compute/ComputeChannel.cpp
--- a/Engine/Source/Developer/Horde/Private/Compute/ComputeChannel.cpp
+++ b/Engine/Source/Developer/Horde/Private/Compute/ComputeChannel.cpp
@@ -10,8 +10,8 @@ FComputeChannel::FComputeChannel()
 }
 
 FComputeChannel::FComputeChannel(FComputeBufferReader InReader, FComputeBufferWriter InWriter)
-	: Reader(std::move(InReader))
-	, Writer(std::move(InWriter))
+	: Reader(MoveTemp(InReader))
+	, Writer(MoveTemp(InWriter))
 {
 }
 
diff --git a/Engine/Source/Developer/Horde/Private/HordePlatform.h b/Engine/Source/Developer/Horde/Private/HordePlatform.h
--- a/Engine/Source/Developer/Horde/Private/HordePlatform.h
+++ b/Engine/Source/Developer/Horde/Private/HordePlatform.h
@@ -14,34 +14,34 @@ struct FHordePlatform
 	//
 
 	// Indicate that the current codepath is not implemented
-	[[noreturn]] HORDE_API static void NotImplemented();
+	[[noreturn]] static void NotImplemented();
 
 	// Indicate that the current codepath is not supported
-	[[noreturn]] HORDE_API static void NotSupported(const char* Message);
+	[[noreturn]] static void NotSupported(const char* Message);
 
 	// Reads an environment variable
-	HORDE_API static bool GetEnvironmentVariable(const char* Name, char* Buffer, size_t BufferLen);
+	static bool GetEnvironmentVariable(const char* Name, char* Buffer, size_t BufferLen);
 
 	// Creates a unique identifier
-	HORDE_API static void CreateUniqueIdentifier(char* NameBuffer, size_t NameBufferLen);
+	static void CreateUniqueIdentifier(char* NameBuffer, size_t NameBufferLen);
 
 	// Creates a unique object name
-	HORDE_API static void CreateUniqueName(char* NameBuffer, size_t NameBufferLen);
+	static void CreateUniqueName(char* NameBuffer, size_t NameBufferLen);
 
 	//
 	// Math
 	//
 
 	// Find the log2 of the given value, returning 0 if the value is zero.
-	HORDE_API static unsigned int FloorLog2(unsigned int Value);
+	static unsigned int FloorLog2(unsigned int Value);
 
 	// Count the number of leading zeros in the given value.
-	HORDE_API static unsigned int CountLeadingZeros(unsigned int Value);
+	static unsigned int CountLeadingZeros(unsigned int Value);
 
 	//
 	// Strings
 	//
 
 	// Parse an integer from a string, returning the number of bytes read
-	HORDE_API static bool TryParseSizeT(const char* Source, size_t SourceLen, size_t& OutValue, size_t& OutNumBytes);
+	static bool TryParseSizeT(const char* Source, size_t SourceLen, size_t& OutValue, size_t& OutNumBytes);
 };
diff --git a/Engine/Source/Developer/Horde/Private/SharedBufferView.cpp b/Engine/Source/Developer/Horde/Private/SharedBufferView.cpp
new file mode 100644
--- /dev/null
+++ b/Engine/Source/Developer/Horde/Private/SharedBufferView.cpp
@@ -0,0 +1,59 @@
+// Copyright Epic Games, Inc. All Rights Reserved.
+
+#include "SharedBufferView.h"
+
+FSharedBufferView::FSharedBufferView()
+{
+}
+
+FSharedBufferView::FSharedBufferView(FSharedBuffer InBuffer)
+	: Buffer(MoveTemp(InBuffer))
+	, View(Buffer.GetView())
+{
+}
+
+FSharedBufferView::FSharedBufferView(FSharedBuffer InBuffer, const FMemoryView& InView)
+	: Buffer(MoveTemp(InBuffer))
+	, View(InView)
+{
+}
+
+FSharedBufferView::FSharedBufferView(FSharedBuffer InBuffer, size_t InOffset, size_t InLength)
+	: Buffer(MoveTemp(InBuffer))
+	, View(Buffer.GetView().Mid(InOffset, InLength))
+{
+}
+
+FSharedBufferView::~FSharedBufferView()
+{
+}
+
+FSharedBufferView FSharedBufferView::Copy(const FMemoryView& Span)
+{
+	return FSharedBufferView(FSharedBuffer::Clone(Span));
+}
+
+FSharedBufferView FSharedBufferView::Slice(uint64 Offset) const
+{
+	return FSharedBufferView(Buffer, View.Mid(Offset));
+}
+
+FSharedBufferView FSharedBufferView::Slice(uint64 Offset, uint64 Length) const
+{
+	return FSharedBufferView(Buffer, View.Mid(Offset, Length));
+}
+
+const unsigned char* FSharedBufferView::GetPointer() const
+{
+	return (const unsigned char*)View.GetData();
+}
+
+size_t FSharedBufferView::GetLength() const
+{
+	return View.GetSize();
+}
+
+FMemoryView FSharedBufferView::GetView() const
+{
+	return View;
+}
diff --git a/Engine/Source/Developer/Horde/Private/Storage/Blob.cpp b/Engine/Source/Developer/Horde/Private/Storage/Blob.cpp
new file mode 100644
--- /dev/null
+++ b/Engine/Source/Developer/Horde/Private/Storage/Blob.cpp
@@ -0,0 +1,13 @@
+// Copyright Epic Games, Inc. All Rights Reserved.
+
+#include "Storage/Blob.h"
+#include "Storage/BlobHandle.h"
+
+FBlob::FBlob(const FBlobType& InType, FSharedBufferView InData, TArray<FBlobHandle> InReferences)
+	: Type(InType)
+	, Data(MoveTemp(InData))
+	, References(MoveTemp(InReferences))
+{ }
+
+FBlob::~FBlob()
+{ }
diff --git a/Engine/Source/Developer/Horde/Private/Storage/BlobHandle.cpp b/Engine/Source/Developer/Horde/Private/Storage/BlobHandle.cpp
new file mode 100644
--- /dev/null
+++ b/Engine/Source/Developer/Horde/Private/Storage/BlobHandle.cpp
@@ -0,0 +1,109 @@
+// Copyright Epic Games, Inc. All Rights Reserved.
+
+#include "Storage/BlobHandle.h"
+#include "Storage/BlobLocator.h"
+#include "Storage/Blob.h"
+#include "../HordePlatform.h"
+
+bool FBlobHandle::operator==(const FBlobHandle& Other) const
+{
+	if (Other.IsValid())
+	{
+		return Get()->Equals(*Other.Get());
+	}
+	else
+	{
+		return !Other.IsValid();
+	}
+}
+
+bool FBlobHandle::operator!=(const FBlobHandle& Other) const
+{
+	return !(*this == Other);
+}
+
+uint32 GetTypeHash(const FBlobHandle& Handle)
+{
+	return Handle->GetHashCode();
+}
+
+// -----------------------------------------------------------------
+
+FBlobHandleData::~FBlobHandleData()
+{
+}
+
+bool FBlobHandleData::Equals(const FBlobHandleData& Other) const
+{
+	return GetType() == Other.GetType();
+}
+
+uint32 FBlobHandleData::GetHashCode() const
+{
+	return (uint32)(size_t)GetType();
+}
+
+FBlobType FBlobHandleData::ReadType() const
+{
+	FBlob Blob = Read();
+	return Blob.Type;
+}
+
+void FBlobHandleData::ReadImports(TArray<FBlobHandle>& OutImports) const
+{
+	FBlob Blob = Read();
+	OutImports = Blob.References;
+}
+
+FSharedBufferView FBlobHandleData::ReadBody() const
+{
+	FBlob Blob = Read();
+	return MoveTemp(Blob.Data);
+}
+
+FSharedBufferView FBlobHandleData::ReadBody(size_t Offset, TOptional<size_t> Length) const
+{
+	FBlob Blob = Read();
+	if (Length.IsSet())
+	{
+		return Blob.Data.Slice(Offset, *Length);
+	}
+	else
+	{
+		return MoveTemp(Blob.Data);
+	}
+}
+
+bool TryAppendLocator(const FBlobHandleData& Handle, FUtf8String& OutPath)
+{
+	FBlobHandle Outer = Handle.GetOuter();
+	if (Outer)
+	{
+		if (!TryAppendLocator(*Outer, OutPath))
+		{
+			return false;
+		}
+
+		if (!Outer->GetOuter())
+		{
+			OutPath += '#';
+		}
+		else
+		{
+			OutPath += '&';
+		}
+	}
+	return Handle.TryAppendIdentifier(OutPath);
+}
+
+FBlobLocator FBlobHandleData::GetLocator() const
+{
+	FUtf8String Path;
+	verify(TryAppendLocator(*this, Path));
+	return FBlobLocator(MoveTemp(Path));
+}
+
+FBlobHandle FBlobHandleData::GetFragmentHandle(const FUtf8StringView& Fragment) const
+{
+	FHordePlatform::NotSupported("This handle type does not support nested blobs");
+}
diff --git a/Engine/Source/Developer/Horde/Private/Storage/BlobLocator.cpp b/Engine/Source/Developer/Horde/Private/Storage/BlobLocator.cpp
new file mode 100644
--- /dev/null
+++ b/Engine/Source/Developer/Horde/Private/Storage/BlobLocator.cpp
@@ -0,0 +1,102 @@
+// Copyright Epic Games, Inc. All Rights Reserved.
+
+#include "Storage/BlobLocator.h"
+#include "Containers/Utf8String.h"
+
+FBlobLocator::FBlobLocator()
+{
+}
+
+FBlobLocator::FBlobLocator(FUtf8String InPath)
+	: Path(MoveTemp(InPath))
+{
+}
+
+FBlobLocator::FBlobLocator(const FBlobLocator& InBaseLocator, const FUtf8StringView& InFragment)
+{
+	if (InFragment.Len() == 0)
+	{
+		Path = InBaseLocator.Path;
+	}
+	else
+	{
+		Path.Reserve(InBaseLocator.Path.Len() + 1 + InFragment.Len());
+		Path += InBaseLocator.Path;
+		Path += '#';
+		Path += InFragment;
+	}
+}
+
+FBlobLocator::FBlobLocator(const FUtf8StringView& InPath)
+	: Path(InPath)
+{
+}
+
+bool FBlobLocator::IsValid() const
+{
+	return !Path.IsEmpty();
+}
+
+const FUtf8String& FBlobLocator::GetPath() const
+{
+	return Path;
+}
+
+FBlobLocator FBlobLocator::GetBaseLocator() const
+{
+	int32 HashIdx;
+	if (Path.FindChar((UTF8CHAR)'#', HashIdx))
+	{
+		return FBlobLocator(Path.Mid(0, HashIdx));
+	}
+	else
+	{
+		return *this;
+	}
+}
+
+FUtf8StringView FBlobLocator::GetFragment() const
+{
+	int32 HashIdx;
+	if (Path.FindChar((UTF8CHAR)'#', HashIdx))
+	{
+		return FUtf8StringView(*Path + HashIdx + 1);
+	}
+	else
+	{
+		return FUtf8StringView();
+	}
+}
+
+bool FBlobLocator::CanUnwrap() const
+{
+	int32 HashIdx;
+	return Path.FindChar((UTF8CHAR)'#', HashIdx);
+}
+
+bool FBlobLocator::TryUnwrap(FBlobLocator& OutLocator, FUtf8StringView& OutFragment) const
+{
+	int32 HashIdx;
+	if (Path.FindChar((UTF8CHAR)'#', HashIdx))
+	{
+		OutLocator = FBlobLocator(Path.Mid(0, HashIdx));
+		OutFragment = FUtf8StringView(*Path + HashIdx + 1);
+		return true;
+	}
+	return false;
+}
+
+bool FBlobLocator::operator==(const FBlobLocator& Other) const
+{
+	return Path.Equals(Other.Path, ESearchCase::CaseSensitive);
+}
+
+bool FBlobLocator::operator!=(const FBlobLocator& Other) const
+{
+	return !(*this == Other);
+}
+
+uint32 GetTypeHash(const FBlobLocator& Locator)
+{
+	return GetTypeHash(Locator.Path);
+}
diff --git a/Engine/Source/Developer/Horde/Private/Storage/BlobReader.cpp b/Engine/Source/Developer/Horde/Private/Storage/BlobReader.cpp
new file mode 100644
--- /dev/null
+++ b/Engine/Source/Developer/Horde/Private/Storage/BlobReader.cpp
@@ -0,0 +1,95 @@
+// Copyright Epic Games, Inc. All Rights Reserved.
+
+#include "Storage/BlobReader.h"
+#include "Containers/Utf8String.h"
+#include "IO/IoHash.h"
+#include "HordePlatform.h"
+
+// ------------------------------------------------------------------------
+
+FBlobReader::FBlobReader(const FBlob& InBlob)
+	: FBlobReader(InBlob.Data.GetView(), InBlob.References)
+{
+}
+
+FBlobReader::FBlobReader(const FMemoryView& InBuffer, const TArray<FBlobHandle>& InImports)
+	: Buffer(InBuffer)
+	, Imports(InImports)
+	, NextImportIdx(0)
+{
+}
+
+const unsigned char* FBlobReader::GetBuffer() const
+{
+	return (const unsigned char*)GetView().GetData();
+}
+
+FMemoryView FBlobReader::GetView() const
+{
+	return Buffer;
+}
+
+void FBlobReader::Advance(size_t Size)
+{
+	Buffer = Buffer.Mid(Size);
+}
+
+FBlobHandle FBlobReader::ReadImport()
+{
+	return Imports[NextImportIdx++];
+}
+
+// ------------------------------------------------------------------------
+
+int ReadInt32(FBlobReader& Reader)
+{
+	int Value = *(const int*)Reader.GetBuffer();
+	Reader.Advance(sizeof(int));
+	return Value;
+}
+
+FIoHash ReadIoHash(FBlobReader& Reader)
+{
+	FIoHash Hash;
+	memcpy(&Hash, Reader.GetBuffer(), sizeof(FIoHash));
+	Reader.Advance(sizeof(FIoHash));
+	return Hash;
+}
+
+FMemoryView ReadFixedLengthBytes(FBlobReader& Reader, size_t Length)
+{
+	FMemoryView View = Reader.GetView();
+	Reader.Advance(Length);
+	return View.Left(Length);
+}
+
+size_t ReadUnsignedVarInt(FBlobReader& Reader)
+{
+	// Figure out the length of the buffer
+	const unsigned char* Data = Reader.GetBuffer();
+	size_t NumBytes = FHordePlatform::CountLeadingZeros((unsigned char)(~*Data)) - 23;
+
+	// Decode the value
+	size_t value = (size_t)(Data[0] & (0xff >> NumBytes));
+	for (int i = 1; i < NumBytes; i++)
+	{
+		value <<= 8;
+		value |= Data[i];
+	}
+
+	Reader.Advance(NumBytes);
+	return value;
+}
+
+FUtf8String ReadString(FBlobReader& Reader)
+{
+	FMemoryView String = ReadStringSpan(Reader);
+	return FUtf8String::ConstructFromPtrSize((const UTF8CHAR*)String.GetData(), String.GetSize());
+}
+
+FMemoryView ReadStringSpan(FBlobReader& Reader)
+{
+	size_t Length = ReadUnsignedVarInt(Reader);
+	return ReadFixedLengthBytes(Reader, Length);
+}
+
diff --git a/Engine/Source/Developer/Horde/Private/Storage/BlobType.cpp b/Engine/Source/Developer/Horde/Private/Storage/BlobType.cpp
new file mode 100644
--- /dev/null
+++ b/Engine/Source/Developer/Horde/Private/Storage/BlobType.cpp
@@ -0,0 +1,20 @@
+// Copyright Epic Games, Inc. All Rights Reserved.
+
+#include "Storage/BlobType.h"
+
+const FBlobType FBlobType::Leaf(FGuid(0x1080B643, 0x4A4D8015, 0xFBC2BBAD, 0x818D6C58), 1);
+
+FBlobType::FBlobType(const FGuid& InGuid, int InVersion)
+	: Guid(InGuid)
+	, Version(InVersion)
+{ }
+
+bool FBlobType::operator==(const FBlobType& Other) const
+{
+	return Guid == Other.Guid && Version == Other.Version;
+}
+
+bool FBlobType::operator!=(const FBlobType& Other) const
+{
+	return !operator==(Other);
+}
\ No newline at end of file
diff --git a/Engine/Source/Developer/Horde/Private/Storage/BlobWriter.cpp b/Engine/Source/Developer/Horde/Private/Storage/BlobWriter.cpp
new file mode 100644
--- /dev/null
+++ b/Engine/Source/Developer/Horde/Private/Storage/BlobWriter.cpp
@@ -0,0 +1,104 @@
+// Copyright Epic Games, Inc. All Rights Reserved.
+
+#include "Storage/BlobWriter.h"
+#include "Storage/RefOptions.h"
+#include "IO/IoHash.h"
+#include "../HordePlatform.h"
+
+FBlobWriter::~FBlobWriter()
+{
+}
+
+void FBlobWriter::AddRef(const FRefName& RefName)
+{
+	AddRef(RefName, FRefOptions());
+}
+
+void* FBlobWriter::GetOutputBuffer(size_t Size)
+{
+	return GetOutputBufferAsSpan(0, Size).GetData();
+}
+
+// ------------------------------------------------------------------------------------------
+
+void WriteIoHash(FBlobWriter& Writer, const FIoHash& Hash)
+{
+	WriteFixedLengthBytes(Writer, &Hash, sizeof(Hash));
+}
+
+void WriteFixedLengthBytes(FBlobWriter& Writer, const void* Data, size_t Length)
+{
+	void* Target = Writer.GetOutputBuffer(Length);
+	memcpy(Target, Data, Length);
+	Writer.Advance(Length);
+}
+
+void WriteFixedLengthBytes(FBlobWriter& Writer, const FMemoryView& View)
+{
+	WriteFixedLengthBytes(Writer, View.GetData(), View.GetSize());
+}
+
+size_t MeasureUnsignedVarInt(size_t Value)
+{
+	check(Value == (unsigned int)Value);
+
+	if (Value == 0)
+	{
+		return 1;
+	}
+	else
+	{
+		return (FHordePlatform::FloorLog2((unsigned int)Value) / 7) + 1;
+	}
+}
+
+size_t WriteUnsignedVarInt(void* Buffer, size_t Value)
+{
+	size_t ByteCount = MeasureUnsignedVarInt(Value);
+	WriteUnsignedVarIntWithKnownLength(Buffer, Value, ByteCount);
+	return ByteCount;
+}
+
+void WriteUnsignedVarInt(FBlobWriter& Writer, size_t Value)
+{
+	size_t ByteCount = MeasureUnsignedVarInt(Value);
+	WriteUnsignedVarIntWithKnownLength(Writer.GetOutputBuffer(ByteCount), Value, ByteCount);
+	Writer.Advance(ByteCount);
+}
+
+void WriteUnsignedVarIntWithKnownLength(void* Buffer, size_t Value, size_t NumBytes)
+{
+	uint8* ByteBuffer = (uint8*)Buffer;
+	for (size_t Idx = 1; Idx < NumBytes; Idx++)
+	{
+		ByteBuffer[NumBytes - Idx] = (unsigned)Value;
+		Value >>= 8;
+	}
+	ByteBuffer[0] = (unsigned char)((0xff << (9 - (int)NumBytes)) | (unsigned char)Value);
+}
+
+size_t MeasureString(const char* Text)
+{
+	return MeasureString(FUtf8StringView((const UTF8CHAR*)Text));
+}
+
+size_t MeasureString(const FUtf8StringView& Text)
+{
+	return MeasureUnsignedVarInt(Text.Len()) + Text.Len();
+}
+
+void WriteString(FBlobWriter& Writer, const char* Text)
+{
+	WriteString(Writer, FUtf8StringView(Text));
+}
+
+void WriteString(FBlobWriter& Writer, const FUtf8StringView& Text)
+{
+	WriteUnsignedVarInt(Writer, Text.Len());
+	WriteFixedLengthBytes(Writer, Text.GetData(), Text.Len());
+}
+
+void WriteString(FBlobWriter& Writer, const FUtf8String& Text)
+{
+	WriteString(Writer, FUtf8StringView(Text));
+}
diff --git a/Engine/Source/Developer/Horde/Private/Storage/Bundles/Bundle.cpp b/Engine/Source/Developer/Horde/Private/Storage/Bundles/Bundle.cpp
new file mode 100644
--- /dev/null
+++ b/Engine/Source/Developer/Horde/Private/Storage/Bundles/Bundle.cpp
@@ -0,0 +1,59 @@
+// Copyright Epic Games, Inc. All Rights Reserved.
+
+#include "Storage/Bundles/Bundle.h"
+
+const FBlobType FBundle::BlobType(FGuid(0x7C5BA294, 0x4F922D21, 0x2F85BE85, 0x1E4CCC48), 1);
+
+FBundleSignature::FBundleSignature(EBundleVersion InVersion, size_t InLength)
+	: Version(InVersion)
+	, Length(InLength)
+{
+}
+
+FBundleSignature FBundleSignature::Read(const void* InData)
+{
+	const unsigned char* Data = (const unsigned char*)InData;
+	if (Data[0] == 'U' && Data[1] == 'E' && Data[2] == 'B' && Data[3] == 'N')
+	{
+		EBundleVersion Version = EBundleVersion::Initial;
+		size_t Length = (Data[4] << 24) | (Data[5] << 16) | (Data[6] << 8) | Data[7];
+		return FBundleSignature(Version, Length);
+	}
+	else if (Data[0] == 'U' && Data[1] == 'B' && Data[2] == 'N')
+	{
+		EBundleVersion Version = (EBundleVersion)Data[3];
+		size_t Length = *(const int*)(Data + 4);
+		return FBundleSignature(Version, Length);
+	}
+	else
+	{
+		check(false);
+		return FBundleSignature((EBundleVersion)0, 0);
+	}
+}
+
+void FBundleSignature::Write(void* InData)
+{
+	unsigned char* Data = (unsigned char*)InData;
+	if (Version == EBundleVersion::Initial)
+	{
+		Data[0] = 'U';
+		Data[1] = 'E';
+		Data[2] = 'B';
+		Data[3] = 'N';
+
+		Data[4] = (unsigned char)(Length >> 24);
+		Data[5] = (unsigned char)(Length >> 16);
+		Data[6] = (unsigned char)(Length >> 8);
+		Data[7] = (unsigned char)(Length);
+	}
+	else
+	{
+		Data[0] = 'U';
+		Data[1] = 'B';
+		Data[2] = 'N';
+		Data[3] = (unsigned char)Version;
+
+		*(int*)(Data + 4) = (int)Length;
+	}
+}
diff --git a/Engine/Source/Developer/Horde/Private/Storage/Bundles/BundleCompression.cpp b/Engine/Source/Developer/Horde/Private/Storage/Bundles/BundleCompression.cpp
new file mode 100644
--- /dev/null
+++ b/Engine/Source/Developer/Horde/Private/Storage/Bundles/BundleCompression.cpp
@@ -0,0 +1,65 @@
+// Copyright Epic Games, Inc. All Rights Reserved.
+
+#include "Storage/Bundles/BundleCompression.h"
+#include "../../HordePlatform.h"
+
+size_t FBundleCompression::GetMaxSize(EBundleCompressionFormat Format, const FMemoryView& Input)
+{
+	switch (Format)
+	{
+	case EBundleCompressionFormat::None:
+		return Input.GetSize();
+	case EBundleCompressionFormat::LZ4:
+		return FCompression::CompressMemoryBound(NAME_LZ4, Input.GetSize());
+	case EBundleCompressionFormat::Gzip:
+		return FCompression::CompressMemoryBound(NAME_Gzip, Input.GetSize());
+	case EBundleCompressionFormat::Oodle:
+		return FCompression::CompressMemoryBound(NAME_Oodle, Input.GetSize());
+	default:
+		FHordePlatform::NotSupported("The specified compression format is not currently supported");
+	}
+}
+
+size_t FBundleCompression::Compress(EBundleCompressionFormat Format, const FMemoryView& Input, const FMutableMemoryView& Output)
+{
+	int32 CompressedSize = Output.GetSize();
+	switch (Format)
+	{
+	case EBundleCompressionFormat::None:
+		check(Output.GetSize() == Input.GetSize());
+		memcpy(Output.GetData(), Input.GetData(), Input.GetSize());
+		return Input.GetSize();
+	case EBundleCompressionFormat::LZ4:
+		verify(FCompression::CompressMemory(NAME_LZ4, Output.GetData(), CompressedSize, Input.GetData(), Input.GetSize()));
+		return CompressedSize;
+	case EBundleCompressionFormat::Gzip:
+		verify(FCompression::CompressMemory(NAME_Gzip, Output.GetData(), CompressedSize, Input.GetData(), Input.GetSize()));
+		return CompressedSize;
+	case EBundleCompressionFormat::Oodle:
+		verify(FCompression::CompressMemory(NAME_Oodle, Output.GetData(), CompressedSize, Input.GetData(), Input.GetSize()));
+		return CompressedSize;
+	default:
+		FHordePlatform::NotSupported("The specified compression format is not currently supported");
+	}
+}
+
+void FBundleCompression::Decompress(EBundleCompressionFormat Format, const FMemoryView& Input, const FMutableMemoryView& Output)
+{
+	switch (Format)
+	{
+	case EBundleCompressionFormat::None:
+		memcpy(Output.GetData(), Input.GetData(), Input.GetSize());
+		break;
+	case EBundleCompressionFormat::LZ4:
+		verify(FCompression::UncompressMemory(NAME_LZ4, Output.GetData(), Output.GetSize(), Input.GetData(), Input.GetSize()));
+		break;
+	case EBundleCompressionFormat::Gzip:
+		verify(FCompression::UncompressMemory(NAME_Gzip, Output.GetData(), Output.GetSize(), Input.GetData(), Input.GetSize()));
+		break;
+	case EBundleCompressionFormat::Oodle:
+		verify(FCompression::UncompressMemory(NAME_Oodle, Output.GetData(), Output.GetSize(), Input.GetData(), Input.GetSize()));
+		break;
+	default:
+		FHordePlatform::NotSupported("The specified compression format is not currently supported");
+	}
+}
diff --git a/Engine/Source/Developer/Horde/Private/Storage/Bundles/BundleOptions.cpp b/Engine/Source/Developer/Horde/Private/Storage/Bundles/BundleOptions.cpp
new file mode 100644
--- /dev/null
+++ b/Engine/Source/Developer/Horde/Private/Storage/Bundles/BundleOptions.cpp
@@ -0,0 +1,9 @@
+// Copyright Epic Games, Inc. All Rights Reserved.
+
+#include "Storage/Bundles/BundleOptions.h"
+
+const FBundleOptions FBundleOptions::Default;
+
+FBundleOptions::FBundleOptions()
+{
+}
diff --git a/Engine/Source/Developer/Horde/Private/Storage/Bundles/V2/BundleWriter.cpp b/Engine/Source/Developer/Horde/Private/Storage/Bundles/V2/BundleWriter.cpp
new file mode 100644
--- /dev/null
+++ b/Engine/Source/Developer/Horde/Private/Storage/Bundles/V2/BundleWriter.cpp
@@ -0,0 +1,440 @@
+// Copyright Epic Games, Inc. All Rights Reserved.
+
+#include "Storage/Bundles/V2/BundleWriter.h"
+#include "Storage/Bundles/V2/PacketWriter.h"
+#include "Storage/Blob.h"
+#include "Storage/StorageClient.h"
+#include "Storage/Bundles/V2/ExportHandle.h"
+#include "../../../HordePlatform.h"
+#include "Storage/ChunkedBufferWriter.h"
+
+//
+// Data for an export that has been written to storage, but may not have been flushed yet
+//
+class FBundleWriter::FPendingExportHandleData : public FBlobHandleData
+{
+public:
+	static const char HandleType[];
+
+	FPendingExportHandleData(TBlobHandle<FPendingPacketHandleData> InPacket, int32 InExportIdx);
+
+	// Overridden from FBlobHandleData
+	bool Equals(const FBlobHandleData& Other) const override;
+	uint32 GetHashCode() const override;
+	FBlobHandle GetOuter() const override;
+	const char* GetType() const override;
+	FBlob Read() const override;
+	bool TryAppendIdentifier(FUtf8String& OutIdentifier) const override;
+
+private:
+	TBlobHandle<FPendingPacketHandleData> Packet;
+	const int32 ExportIdx;
+};
+
+//
+// Packet which may not have been flushed yet
+//
+class FBundleWriter::FPendingPacketHandleData : public FBlobHandleData
+{
+public:
+	static const char HandleType[];
+
+	FPendingPacketHandleData(FBundleWriter& InBundleWriter, TBlobHandle<FPendingBundleHandleData> InBundle);
+
+	// Reads an export from the packet. The packet may or may not be flushed.
+	FBlob ReadExport(int ExportIdx) const;
+
+	// Marks the current packet as complete, and returns the packet writer for reuse on the next packet
+	void MarkComplete(TBlobHandle<FPacketHandleData> InFlushedHandle);
+
+	// Overridden from FBlobHandleData
+	bool Equals(const FBlobHandleData& Other) const override;
+	uint32 GetHashCode() const override;
+	FBlobHandle GetOuter() const override;
+	const char* GetType() const override;
+	FBlob Read() const override;
+	bool TryAppendIdentifier(FUtf8String& OutIdentifier) const override;
+
+private:
+	FBundleWriter* BundleWriter;
+	TBlobHandle<FPendingBundleHandleData> BundleHandle;
+	TBlobHandle<FPacketHandleData> FlushedHandle;
+};
+
+//
+// Bundle which is still being written to
+//
+class FBundleWriter::FPendingBundleHandleData : public FBlobHandleData, public TSharedFromThis<FBundleWriter::FPendingBundleHandleData>
+{
+public:
+	static const char HandleType[];
+
+	mutable FCriticalSection CriticalSection;
+
+	FPendingBundleHandleData(FBundleWriter& InBundleWriter);
+	void MarkComplete(FBlobHandle InFlushedHandle);
+
+	// Overridden from FBlobHandleData
+	bool Equals(const FBlobHandleData& Other) const override;
+	uint32 GetHashCode() const override;
+	FBlobHandle GetOuter() const override;
+	const char* GetType() const override;
+	FBlob Read() const override;
+	FSharedBufferView ReadBody(size_t Offset, TOptional<size_t> Length = TOptional<size_t>()) const override;
+	bool TryAppendIdentifier(FUtf8String& OutIdentifier) const override;
+
+private:
+	FBundleWriter* BundleWriter;
+	FBlobHandle FlushedHandle;
+};
+
+
+
+
+
+// --------------------------------------------------------------------------------------------------------
+
+const char FBundleWriter::FPendingExportHandleData::HandleType[] = "FBundleWriter::FPendingExportHandleData";
+
+FBundleWriter::FPendingExportHandleData::FPendingExportHandleData(TBlobHandle<FPendingPacketHandleData> InPacket, int32 InExportIdx)
+	: Packet(MoveTemp(InPacket))
+	, ExportIdx(InExportIdx)
+{
+}
+
+bool FBundleWriter::FPendingExportHandleData::Equals(const FBlobHandleData& Other) const
+{
+	return this == &Other;
+}
+
+uint32 FBundleWriter::FPendingExportHandleData::GetHashCode() const
+{
+	return (uint32)(size_t)this;
+}
+
+FBlobHandle FBundleWriter::FPendingExportHandleData::GetOuter() const
+{
+	return Packet;
+}
+
+const char* FBundleWriter::FPendingExportHandleData::GetType() const
+{
+	return HandleType;
+}
+
+FBlob FBundleWriter::FPendingExportHandleData::Read() const
+{
+	return Packet->ReadExport(ExportIdx);
+}
+
+bool FBundleWriter::FPendingExportHandleData::TryAppendIdentifier(FUtf8String& OutIdentifier) const
+{
+	FExportHandleData::AppendIdentifier(OutIdentifier, ExportIdx);
+	return true;
+}
+
+
+
+
+
+
+// --------------------------------------------------------------------------------------------------------
+
+const char FBundleWriter::FPendingPacketHandleData::HandleType[] = "FBundleWriter::FPendingPacketHandleData";
+
+FBundleWriter::FPendingPacketHandleData::FPendingPacketHandleData(FBundleWriter& InBundleWriter, TBlobHandle<FPendingBundleHandleData> InBundleHandle)
+	: BundleWriter(&InBundleWriter)
+	, BundleHandle(MoveTemp(InBundleHandle))
+{
+}
+
+void FBundleWriter::FPendingPacketHandleData::MarkComplete(TBlobHandle<FPacketHandleData> InFlushedHandle)
+{
+	check(!FlushedHandle);
+	BundleWriter = nullptr;
+	FlushedHandle = MoveTemp(InFlushedHandle);
+}
+
+bool FBundleWriter::FPendingPacketHandleData::Equals(const FBlobHandleData& Other) const
+{
+	return this == &Other;
+}
+
+uint32 FBundleWriter::FPendingPacketHandleData::GetHashCode() const
+{
+	return (uint32)(size_t)this;
+}
+
+FBlobHandle FBundleWriter::FPendingPacketHandleData::GetOuter() const
+{
+	return FlushedHandle.IsValid() ? (FBlobHandle)FlushedHandle->GetOuter() : (FBlobHandle)BundleHandle;
+}
+
+const char* FBundleWriter::FPendingPacketHandleData::GetType() const
+{
+	return HandleType;
+}
+
+FBlob FBundleWriter::FPendingPacketHandleData::Read() const
+{
+	{
+		FScopeLock Lock(&BundleHandle->CriticalSection);
+		check(FlushedHandle);
+	}
+	return FlushedHandle->Read();
+}
+
+FBlob FBundleWriter::FPendingPacketHandleData::ReadExport(int ExportIdx) const
+{
+	{
+		FScopeLock Lock(&BundleHandle->CriticalSection);
+		if (!FlushedHandle)
+		{
+			// If not flushed yet, the packet writer must be this packet
+			check(BundleWriter->CurrentPacketHandle.Get() == this);
+			return BundleWriter->PacketWriter->GetExport(ExportIdx);
+		}
+	}
+	return FlushedHandle->ReadExport(ExportIdx);
+}
+
+bool FBundleWriter::FPendingPacketHandleData::TryAppendIdentifier(FUtf8String& OutIdentifier) const
+{
+	if (FlushedHandle)
+	{
+		return FlushedHandle->TryAppendIdentifier(OutIdentifier);
+	}
+	else
+	{
+		return false;
+	}
+}
+
+
+
+
+
+// --------------------------------------------------------------------------------------------------------
+
+const char FBundleWriter::FPendingBundleHandleData::HandleType[] = "FBundleWriter::FPendingBundleHandleData";
+
+FBundleWriter::FPendingBundleHandleData::FPendingBundleHandleData(FBundleWriter& InBundleWriter)
+	: BundleWriter(&InBundleWriter)
+{
+}
+
+void FBundleWriter::FPendingBundleHandleData::MarkComplete(FBlobHandle InFlushedHandle)
+{
+	FScopeLock Lock(&CriticalSection);
+	check(!FlushedHandle);
+	FlushedHandle = InFlushedHandle;
+	BundleWriter = nullptr;
+}
+
+bool FBundleWriter::FPendingBundleHandleData::Equals(const FBlobHandleData& Other) const
+{
+	return this == &Other;
+}
+
+uint32 FBundleWriter::FPendingBundleHandleData::GetHashCode() const
+{
+	return (uint32)(size_t)this;
+}
+
+FBlobHandle FBundleWriter::FPendingBundleHandleData::GetOuter() const
+{
+	return FBlobHandle();
+}
+
+const char* FBundleWriter::FPendingBundleHandleData::GetType() const
+{
+	return HandleType;
+}
+
+FBlob FBundleWriter::FPendingBundleHandleData::Read() const
+{
+	return FlushedHandle->Read();
+}
+
+FSharedBufferView FBundleWriter::FPendingBundleHandleData::ReadBody(size_t Offset, TOptional<size_t> Length) const
+{
+	if (!FlushedHandle)
+	{
+		FScopeLock Lock(&CriticalSection);
+		if (!FlushedHandle)
+		{
+			size_t LengthValue = Length.IsSet() ? Length.GetValue() : BundleWriter->CompressedPacketWriter.GetLength() - Offset;
+			return BundleWriter->CompressedPacketWriter.Slice(Offset, LengthValue);
+		}
+	}
+	return FlushedHandle->ReadBody(Offset, Length);
+}
+
+bool FBundleWriter::FPendingBundleHandleData::TryAppendIdentifier(FUtf8String& OutIdentifier) const
+{
+	if (FlushedHandle)
+	{
+		FlushedHandle->TryAppendIdentifier(OutIdentifier);
+		return true;
+	}
+	return false;
+}
+
+
+
+
+
+
+
+
+// --------------------------------------------------------------------------------------------------------
+
+FBundleWriter::FBundleWriter(TSharedRef<FKeyValueStorageClient> InStorageClient, FUtf8String InBasePath, const FBundleOptions& InOptions)
+	: StorageClient(MoveTemp(InStorageClient))
+	, BasePath(MoveTemp(InBasePath))
+	, Options(InOptions)
+	, CurrentBundleHandle(MakeShared<FPendingBundleHandleData>(*this))
+{
+	StartPacket();
+}
+
+FBundleWriter::~FBundleWriter()
+{
+}
+
+void FBundleWriter::AddAlias(const FAliasInfo& Alias)
+{
+	FHordePlatform::NotImplemented();
+}
+
+void FBundleWriter::AddImport(FBlobHandle Handle)
+{
+	PacketWriter->AddImport(MoveTemp(Handle));
+}
+
+void FBundleWriter::AddRef(const FRefName& RefName, const FRefOptions& RefOptions)
+{
+	FHordePlatform::NotImplemented();
+}
+
+FBlobHandle FBundleWriter::CompleteBlob(const FBlobType& Type)
+{
+	check(PacketWriter);
+	check(CurrentPacketHandle);
+	check(CurrentBundleHandle);
+
+	int ExportIdx = PacketWriter->CompleteBlob(Type);
+	FBlobHandle ExportHandle(MakeShared<FPendingExportHandleData>(CurrentPacketHandle, ExportIdx));
+
+//	for (FAliasInfo& Alias : Aliases)
+//	{
+//		PendingExportAliases.push_back(std::make_pair(ExportHandle, std::move(Alias)));
+//	}
+
+	if (PacketWriter->GetLength() > FMath::Min(Options.MinCompressionPacketSize, Options.MaxBlobSize))
+	{
+		FinishPacket();
+		StartPacket();
+	}
+
+	return ExportHandle;
+}
+
+FMutableMemoryView FBundleWriter::GetOutputBufferAsSpan(size_t UsedSize, size_t DesiredSize)
+{
+	return PacketWriter->GetOutputBuffer(UsedSize, DesiredSize);
+}
+
+void FBundleWriter::Advance(size_t Size)
+{
+	PacketWriter->Advance(Size);
+}
+
+void FBundleWriter::Flush()
+{
+	// Check we haven't already flushed this bundle
+	if (CurrentPacketHandle)
+	{
+		FinishPacket();
+	}
+	if (CompressedPacketWriter.GetLength() == 0)
+	{
+		return;
+	}
+
+	// Write the data to the underlying store
+	TArray<FMemoryView> View = CompressedPacketWriter.GetView();
+	FBlobHandle FlushedHandle = StorageClient->WriteBlob(BasePath, FBundle::BlobType, View, CurrentBundleImports);
+
+	// Update the bundle and packet handles
+	{
+		FScopeLock Lock(&CurrentBundleHandle->CriticalSection);
+		CurrentBundleHandle->MarkComplete(FlushedHandle);
+		CurrentBundleHandle = MakeShared<FPendingBundleHandleData>(*this);
+	}
+
+	// Reset the writer
+	CompressedPacketWriter.Reset();
+	CurrentBundleImports.Empty();
+
+	// Create a new bundle
+
+	// TODO: put all the encoded packets into the cache using the final handles
+
+//	// Add all the aliases
+//	for(std::pair<FBlobHandle, FAliasInfo>& PendingExportAlias : PendingExportAliases)
+//	{
+//		FBlobAlias Alias(PendingExportAlias.first, PendingExportAlias.second.Rank, PendingExportAlias.second.Data);
+//		StorageClient->AddAlias(PendingExportAlias.second.Name.c_str(), MoveTemp(Alias));
+//	}
+}
+
+TUniquePtr<FBlobWriter> FBundleWriter::Fork()
+{
+	return TUniquePtr<FBlobWriter>(new FBundleWriter(StorageClient, BasePath, Options));
+}
+
+void FBundleWriter::StartPacket()
+{
+	check(CurrentBundleHandle);
+	check(!CurrentPacketHandle);
+	check(!PacketWriter);
+
+	CurrentPacketHandle = MakeShared<FPendingPacketHandleData>(*this, CurrentBundleHandle);
+	PacketWriter = MakeUnique<FPacketWriter>(CurrentBundleHandle, CurrentPacketHandle);
+}
+
+void FBundleWriter::FinishPacket()
+{
+	check(CurrentBundleHandle);
+	check(CurrentPacketHandle);
+	check(PacketWriter);
+
+	if (PacketWriter->GetExportCount() > 0)
+	{
+		size_t PacketOffset = CompressedPacketWriter.GetLength();
+		FPacket Packet = PacketWriter->CompletePacket();
+		Packet.Encode(Options.CompressionFormat, CompressedPacketWriter);
+		size_t PacketLength = CompressedPacketWriter.GetLength() - PacketOffset;
+
+		// Point the packet handle to the encoded data
+		{
+			FScopeLock Lock(&CurrentBundleHandle->CriticalSection);
+			TBlobHandle<FPacketHandleData> FlushedPacketHandle = MakeShared<FPacketHandleData>(StorageClient, CurrentBundleHandle, PacketOffset, PacketLength);
+			CurrentPacketHandle->MarkComplete(MoveTemp(FlushedPacketHandle));
+		}
+
+		// Find all the other bundles that are referenced
+		for (int ImportIdx = 0; ImportIdx < Packet.GetImportCount(); ImportIdx++)
+		{
+			FPacketImport Import = Packet.GetImport(ImportIdx);
+			if (Import.GetBaseIdx() == FPacketImport::InvalidBaseIdx)
+			{
+				CurrentBundleImports.Add(PacketWriter->GetImport(ImportIdx));
+			}
+		}
+	}
+
+	PacketWriter.Reset();
+	CurrentPacketHandle = TBlobHandle<FPendingPacketHandleData>();
+}
diff --git a/Engine/Source/Developer/Horde/Private/Storage/Bundles/V2/ExportHandle.cpp b/Engine/Source/Developer/Horde/Private/Storage/Bundles/V2/ExportHandle.cpp
new file mode 100644
--- /dev/null
+++ b/Engine/Source/Developer/Horde/Private/Storage/Bundles/V2/ExportHandle.cpp
@@ -0,0 +1,91 @@
+// Copyright Epic Games, Inc. All Rights Reserved.
+
+#include "Storage/Bundles/V2/ExportHandle.h"
+#include "../../../HordePlatform.h"
+
+const char FExportHandleData::Type[] = "FExportHandleData";
+const FUtf8StringView FExportHandleData::FragmentPrefix("exp=");
+
+FExportHandleData::FExportHandleData(FPacketHandle InPacket, int32 InExportIdx)
+	: Packet(MoveTemp(InPacket))
+	, ExportIdx(InExportIdx)
+{
+}
+
+FExportHandleData::FExportHandleData(FPacketHandle InPacket, const FUtf8StringView& InFragment)
+	: Packet(MoveTemp(InPacket))
+{
+	verify(TryParse(InFragment, ExportIdx));
+}
+
+FExportHandleData::~FExportHandleData()
+{
+}
+
+void FExportHandleData::AppendIdentifier(FUtf8String& OutBuffer, int32 ExportIdx)
+{
+	OutBuffer.Appendf("exp=%d", ExportIdx);
+}
+
+bool FExportHandleData::Equals(const FBlobHandleData& Other) const
+{
+	if (GetType() != Other.GetType())
+	{
+		return false;
+	}
+
+	const FExportHandleData& OtherExport = (const FExportHandleData&)Other;
+	return Packet->Equals(*OtherExport.Packet) && ExportIdx == OtherExport.ExportIdx;
+}
+
+uint32 FExportHandleData::GetHashCode() const
+{
+	return (Packet->GetHashCode() * 257) + ExportIdx;
+}
+
+FBlobHandle FExportHandleData::GetOuter() const
+{
+	return Packet;
+}
+
+const char* FExportHandleData::GetType() const
+{
+	return Type;
+}
+
+FBlob FExportHandleData::Read() const
+{
+	return Packet->ReadExport(ExportIdx);
+}
+
+bool FExportHandleData::TryAppendIdentifier(FUtf8String& OutBuffer) const
+{
+	AppendIdentifier(OutBuffer, ExportIdx);
+	return true;
+}
+
+bool FExportHandleData::TryParse(const FUtf8StringView& Fragment, int32& OutExportIdx)
+{
+	OutExportIdx = 0;
+
+	if (!Fragment.StartsWith(FragmentPrefix, ESearchCase::CaseSensitive))
+	{
+		return false;
+	}
+
+	const UTF8CHAR* Data = Fragment.GetData();
+	for (size_t Idx = FragmentPrefix.Len(); Idx < Fragment.Len(); Idx++)
+	{
+		UTF8CHAR Char = Data[Idx];
+		if (Char >= '0' && Char <= '9')
+		{
+			OutExportIdx = (OutExportIdx * 10) + (Char - '0');
+		}
+		else
+		{
+			return false;
+		}
+	}
+
+	return true;
+}
diff --git a/Engine/Source/Developer/Horde/Private/Storage/Bundles/V2/Packet.cpp b/Engine/Source/Developer/Horde/Private/Storage/Bundles/V2/Packet.cpp
new file mode 100644
--- /dev/null
+++ b/Engine/Source/Developer/Horde/Private/Storage/Bundles/V2/Packet.cpp
@@ -0,0 +1,214 @@
+// Copyright Epic Games, Inc. All Rights Reserved.
+
+#include "Storage/Bundles/V2/Packet.h"
+#include "Storage/Bundles/Bundle.h"
+#include "Storage/BlobWriter.h"
+#include "Storage/ChunkedBufferWriter.h"
+#include "../../../HordePlatform.h"
+#include "Serialization/VarInt.h"
+
+const FBlobType FPacket::BlobType(FGuid(0xCD9A04EF, 0x47D3CAC1, 0x492A05A6, 0x51E63081), 1);
+
+FPacket::FPacket(FSharedBufferView InBuffer)
+	: Buffer(MoveTemp(InBuffer))
+{
+}
+
+FPacket::~FPacket()
+{
+}
+
+FSharedBufferView FPacket::GetBuffer() const
+{
+	return Buffer;
+}
+
+size_t FPacket::GetLength() const
+{
+	return Buffer.GetLength();
+}
+
+size_t FPacket::GetTypeCount() const
+{
+	const unsigned char* Data = Buffer.GetPointer();
+	int TypeTableOffset = *(const int*)(Data + 8);
+	return *(int*)(Data + TypeTableOffset);
+}
+
+const FBlobType& FPacket::GetType(size_t TypeIdx) const
+{
+	const unsigned char* Data = Buffer.GetPointer();
+	int TypeTableOffset = *(const int*)(Data + 8);
+
+	int TypeCount = *(Data + TypeTableOffset);
+	check(TypeIdx >= 0 && TypeIdx < TypeCount);
+
+	const FBlobType* Types = (const FBlobType*)(Data + TypeTableOffset + sizeof(int));
+	return Types[TypeIdx];
+}
+
+size_t FPacket::GetImportCount() const
+{
+	const unsigned char* Data = Buffer.GetPointer();
+	int ImportTableOffset = *(const int*)(Data + 12);
+	return *(const int*)(Data + ImportTableOffset);
+}
+
+FPacketImport FPacket::GetImport(size_t ImportIdx) const
+{
+	const unsigned char* Data = Buffer.GetPointer();
+	size_t ImportTableOffset = *(const int*)(Data + 12);
+
+	size_t ImportCount = *(const int*)(Data + ImportTableOffset);
+	check(ImportIdx >= 0 && ImportIdx < ImportCount);
+
+	size_t ImportEntryOffset = ImportTableOffset + sizeof(int) + ImportIdx * sizeof(int);
+
+	size_t ImportOffset = *(const int*)(Data + ImportEntryOffset);
+	size_t ImportLength = *(const int*)(Data + ImportEntryOffset + 4) - ImportOffset;
+
+	return FPacketImport(Buffer.GetView().Mid(ImportOffset, ImportLength));
+}
+
+size_t FPacket::GetExportCount() const
+{
+	const unsigned char* Data = Buffer.GetPointer();
+	size_t ExportTableOffset = *(const int*)(Data + 16);
+	return *(const int*)(Data + ExportTableOffset);
+}
+
+FPacketExport FPacket::GetExport(size_t ExportIdx) const
+{
+	const unsigned char* Data = Buffer.GetPointer();
+	size_t ExportTableOffset = *(const int*)(Data + 16);
+
+	size_t ExportCount = *(const int*)(Data + ExportTableOffset);
+	check(ExportIdx >= 0 && ExportIdx < ExportCount);
+
+	size_t ExportEntryOffset = ExportTableOffset + sizeof(int) + ExportIdx * sizeof(int);
+
+	size_t ExportOffset = *(const int*)(Data + ExportEntryOffset);
+	size_t ExportLength = *(const int*)(Data + ExportEntryOffset + 4) - ExportOffset;
+
+	return FPacketExport(Buffer.Slice(ExportOffset, ExportLength));
+}
+
+void FPacket::Encode(EBundleCompressionFormat Format, FChunkedBufferWriter& Writer)
+{
+	FMemoryView View = Buffer.GetView();
+
+	size_t HeaderLength = FBundleSignature::NumBytes + 5;
+	FMutableMemoryView Output = Writer.GetOutputBuffer(0, HeaderLength + FBundleCompression::GetMaxSize(Format, View));
+
+	unsigned char* OutputData = (unsigned char*)Output.GetData();
+	size_t EncodedLength = FBundleCompression::Compress(Format, View, Output.Mid(HeaderLength));
+
+	FBundleSignature(EBundleVersion::LatestV2, HeaderLength + EncodedLength).Write(OutputData);
+	OutputData += FBundleSignature::NumBytes;
+
+	*(int*)OutputData = (int)GetLength();
+	OutputData += sizeof(int);
+
+	*OutputData = (unsigned char)Format;
+
+	Writer.Advance(HeaderLength + EncodedLength);
+}
+
+FPacket FPacket::Decode(const FMemoryView& View)
+{
+	const unsigned char* Data = (const unsigned char*)View.GetData();
+
+	FBundleSignature Signature = FBundleSignature::Read(Data);
+	Data += FBundleSignature::NumBytes;
+
+	if (Signature.Version <= EBundleVersion::LatestV1 || Signature.Version > EBundleVersion::LatestV2)
+	{
+		FHordePlatform::NotSupported("Unsupported bundle version");
+	}
+
+	int DecodedLength = *(int*)Data;
+	Data += 4;
+
+	EBundleCompressionFormat Format = (EBundleCompressionFormat)*Data;
+	Data++;
+
+	FUniqueBuffer Buffer = FUniqueBuffer::Alloc(DecodedLength);
+	FBundleCompression::Decompress(Format, View.Left(Signature.Length).Mid(FBundleSignature::NumBytes + 5), Buffer.GetView());
+
+	return FPacket(Buffer.MoveToShared());
+}
+
+// --------------------------------------------------------------------------------------------------------
+
+FPacketImport::FPacketImport(const FMemoryView& InView)
+{
+	uint32 ByteCount;
+	BaseIdx = (int64)ReadVarUInt(InView.GetData(), ByteCount) - Bias;
+
+	FMemoryView FragmentData = InView.Mid(ByteCount);
+	Fragment = FUtf8StringView((const UTF8CHAR*)FragmentData.GetData(), FragmentData.GetSize());
+}
+
+FPacketImport::FPacketImport(int32 InBaseIdx, const FUtf8StringView& InFragment)
+	: BaseIdx(InBaseIdx)
+	, Fragment(InFragment)
+{
+}
+
+int FPacketImport::GetBaseIdx() const
+{
+	return BaseIdx;
+}
+
+FUtf8StringView FPacketImport::GetFragment() const
+{
+	return Fragment;
+}
+
+// --------------------------------------------------------------------------------------------------------
+
+FPacketExport::FPacketExport(FSharedBufferView InBuffer)
+	: Buffer(MoveTemp(InBuffer))
+{
+	const char* Data = (const char*)Buffer.GetPointer();
+
+	uint32 Offset = sizeof(uint32) + *(const uint32*)Data;
+	uint32 ByteCount;
+
+	TypeIdx = (int32)ReadVarUInt(Data + Offset, ByteCount);
+	Offset += ByteCount;
+
+	Imports = InBuffer.GetView().Mid(Offset);
+}
+
+FPacketExport::~FPacketExport()
+{
+}
+
+size_t FPacketExport::GetTypeIndex() const
+{
+	return TypeIdx;
+}
+
+void FPacketExport::GetImportIndices(TArray<size_t>& OutIndices) const
+{
+	uint32 ByteCount;
+	const char* Data = (const char*)Imports.GetData();
+
+	size_t NumImports = ReadVarUInt(Data, ByteCount);
+	Data += ByteCount;
+	OutIndices.Reserve(OutIndices.Num() + NumImports);
+
+	for (size_t Idx = 0; Idx < NumImports; Idx++)
+	{
+		size_t Index = ReadVarUInt(Data, ByteCount);
+		Data += ByteCount;
+		OutIndices.Add(Index);
+	}
+}
+
+FSharedBufferView FPacketExport::GetPayload() const
+{
+	int PayloadLength = *(int*)Buffer.GetPointer();
+	return Buffer.Slice(sizeof(int), PayloadLength);
+}
diff --git a/Engine/Source/Developer/Horde/Private/Storage/Bundles/V2/PacketHandle.cpp b/Engine/Source/Developer/Horde/Private/Storage/Bundles/V2/PacketHandle.cpp
new file mode 100644
--- /dev/null
+++ b/Engine/Source/Developer/Horde/Private/Storage/Bundles/V2/PacketHandle.cpp
@@ -0,0 +1,129 @@
+// Copyright Epic Games, Inc. All Rights Reserved.
+
+#include "Storage/Bundles/V2/PacketHandle.h"
+#include "Storage/Bundles/V2/PacketReader.h"
+#include "Storage/Bundles/V2/ExportHandle.h"
+#include "../../../HordePlatform.h"
+
+// --------------------------------------------------------------------------------------------------------
+
+const char FPacketHandleData::Type[] = "FPacketHandleDetail";
+const char FPacketHandleData::FragmentPrefix[] = "pkt=";
+const size_t FPacketHandleData::FragmentPrefixLength = (sizeof(FragmentPrefix) / sizeof(FragmentPrefix[0])) - 1;
+
+FPacketHandleData::FPacketHandleData(TSharedRef<FStorageClient> InStorageClient, FBlobHandle InBundle, size_t InPacketOffset, size_t InPacketLength)
+	: StorageClient(MoveTemp(InStorageClient))
+	, Bundle(MoveTemp(InBundle))
+	, PacketOffset(InPacketOffset)
+	, PacketLength(InPacketLength)
+{
+}
+
+FPacketHandleData::FPacketHandleData(TSharedRef<FStorageClient> InStorageClient, FBlobHandle InBundle, const FUtf8StringView& InFragment)
+	: StorageClient(MoveTemp(InStorageClient))
+	, Bundle(MoveTemp(InBundle))
+{
+	verify(TryParse(InFragment, PacketOffset, PacketLength));
+}
+
+FPacketHandleData::~FPacketHandleData()
+{
+}
+
+FBlob FPacketHandleData::ReadExport(size_t ExportIdx) const
+{
+	return GetPacketReader().ReadExport(ExportIdx);
+}
+
+FSharedBufferView FPacketHandleData::ReadExportBody(size_t ExportIdx) const
+{
+	return GetPacketReader().ReadExportBody(ExportIdx);
+}
+
+bool FPacketHandleData::Equals(const FBlobHandleData& Other) const
+{
+	if (GetType() != Other.GetType())
+	{
+		return false;
+	}
+
+	const FPacketHandleData& OtherPacket = (const FPacketHandleData&)Other;
+	return Bundle->Equals(*OtherPacket.Bundle) && PacketOffset == OtherPacket.PacketOffset;
+}
+
+uint32 FPacketHandleData::GetHashCode() const
+{
+	return (Bundle->GetHashCode() * 257) + PacketOffset;
+}
+
+FBlobHandle FPacketHandleData::GetOuter() const
+{
+	return Bundle;
+}
+
+const char* FPacketHandleData::GetType() const
+{
+	return Type;
+}
+
+FBlob FPacketHandleData::Read() const
+{
+	const FPacketReader& PacketReaderInst = GetPacketReader();
+	return PacketReaderInst.ReadPacket();
+}
+
+bool FPacketHandleData::TryAppendIdentifier(FUtf8String& OutBuffer) const
+{
+	OutBuffer.Appendf("pkt=%zu,%zu", PacketOffset, PacketLength);
+	return true;
+}
+
+FBlobHandle FPacketHandleData::GetFragmentHandle(const FUtf8StringView& Fragment) const
+{
+	TBlobHandle<FPacketHandleData> PacketHandle(const_cast<FPacketHandleData*>(this)->AsShared());
+	return MakeShared<FExportHandleData>(PacketHandle, Fragment);
+}
+
+const FPacketReader& FPacketHandleData::GetPacketReader() const
+{
+	if (!PacketReader)
+	{
+		FSharedBufferView Buffer = Bundle->ReadBody(PacketOffset, PacketLength);
+		FPacket Packet = FPacket::Decode(Buffer.GetView());
+		PacketReader = MakeShared<FPacketReader>(const_cast<FPacketHandleData*>(this)->StorageClient, Bundle, const_cast<FPacketHandleData*>(this)->AsShared(), MoveTemp(Packet));
+	}
+	return *PacketReader.Get();
+}
+
+bool FPacketHandleData::TryParse(const FUtf8StringView& Fragment, size_t& OutPacketOffset, size_t& OutPacketLength)
+{
+	size_t NumBytesRead;
+	if (!Fragment.StartsWith(FragmentPrefix))
+	{
+		OutPacketOffset = OutPacketLength = 0;
+		return false;
+	}
+
+	FUtf8StringView Remaining = Fragment.Mid(FragmentPrefixLength);
+	if (!FHordePlatform::TryParseSizeT((const char*)Remaining.GetData(), Remaining.Len(), OutPacketOffset, NumBytesRead))
+	{
+		OutPacketOffset = OutPacketLength = 0;
+		return false;
+	}
+
+	Remaining = Remaining.Mid(NumBytesRead);
+	if (Remaining.Len() == 0 || Remaining[0] != ',')
+	{
+		OutPacketOffset = OutPacketLength = 0;
+		return false;
+	}
+
+	Remaining = Remaining.Mid(1);
+	if (!FHordePlatform::TryParseSizeT((const char*)Remaining.GetData(), Remaining.Len(), OutPacketLength, NumBytesRead) || NumBytesRead != Remaining.Len())
+	{
+		OutPacketOffset = OutPacketLength = 0;
+		return false;
+	}
+
+	return true;
+}
diff --git a/Engine/Source/Developer/Horde/Private/Storage/Bundles/V2/PacketReader.cpp b/Engine/Source/Developer/Horde/Private/Storage/Bundles/V2/PacketReader.cpp
new file mode 100644
--- /dev/null
+++ b/Engine/Source/Developer/Horde/Private/Storage/Bundles/V2/PacketReader.cpp
@@ -0,0 +1,86 @@
+// Copyright Epic Games, Inc. All Rights Reserved.
+
+#include "Storage/Bundles/V2/PacketReader.h"
+#include "Storage/Bundles/V2/ExportHandle.h"
+#include "Storage/Bundles/V2/PacketHandle.h"
+#include "Storage/Blob.h"
+#include "../../../HordePlatform.h"
+
+FPacketReader::FPacketReader(TSharedRef<FStorageClient> InStorageClient, FBlobHandle InBundleHandle, FPacketHandle InPacketHandle, FPacket InPacket)
+	: StorageClient(MoveTemp(InStorageClient))
+	, BundleHandle(MoveTemp(InBundleHandle))
+	, PacketHandle(MoveTemp(InPacketHandle))
+	, Packet(MoveTemp(InPacket))
+{
+	CachedImportHandles.SetNum(Packet.GetImportCount());
+}
+
+FPacketReader::~FPacketReader()
+{
+}
+
+FBlob FPacketReader::ReadPacket() const
+{
+	for (int Idx = 0; Idx < CachedImportHandles.Num(); Idx++)
+	{
+		GetImportHandle(Idx);
+	}
+	return FBlob(FPacket::BlobType, Packet.GetBuffer(), CachedImportHandles);
+}
+
+FBlob FPacketReader::ReadExport(size_t ExportIdx) const
+{
+	FPacketExport Export = Packet.GetExport(ExportIdx);
+
+	FBlobType Type = Packet.GetType(Export.GetTypeIndex());
+	FSharedBufferView Data = Export.GetPayload();
+
+	TArray<size_t> ImportIndices;
+	Export.GetImportIndices(ImportIndices);
+
+	TArray<FBlobHandle> ImportHandles;
+	ImportHandles.Reserve(ImportIndices.Num());
+
+	for (size_t ImportIndex : ImportIndices)
+	{
+		ImportHandles.Add(GetImportHandle(ImportIndex));
+	}
+
+	return FBlob(Type, Export.GetPayload(), MoveTemp(ImportHandles));
+}
+
+FSharedBufferView FPacketReader::ReadExportBody(size_t ExportIdx) const
+{
+	FPacketExport Export = Packet.GetExport(ExportIdx);
+	return Export.GetPayload();
+}
+
+FBlobHandle FPacketReader::GetImportHandle(size_t Index) const
+{
+	FBlobHandle ImportHandle = CachedImportHandles[Index];
+	if (!ImportHandle)
+	{
+		FPacketImport Import = Packet.GetImport(Index);
+
+		size_t BaseIdx = Import.GetBaseIdx();
+		switch (BaseIdx)
+		{
+		case FPacketImport::InvalidBaseIdx:
+			ImportHandle = StorageClient->CreateHandle(FBlobLocator(FUtf8String(Import.GetFragment())));
+			break;
+		case FPacketImport::CurrentBundleBaseIdx:
+			ImportHandle = MakeShared<FPacketHandleData>(StorageClient, BundleHandle, Import.GetFragment());
+			break;
+		case FPacketImport::CurrentPacketBaseIdx:
+			ImportHandle = MakeShared<FExportHandleData>(PacketHandle, Import.GetFragment());
+			break;
+		default:
+			FBlobHandle ParentHandle = GetImportHandle(BaseIdx);
+			ImportHandle = ParentHandle->GetFragmentHandle(Import.GetFragment());
+			break;
+		}
+
+		CachedImportHandles[Index] = ImportHandle;
+	}
+	return ImportHandle;
+}
diff --git a/Engine/Source/Developer/Horde/Private/Storage/Bundles/V2/PacketWriter.cpp b/Engine/Source/Developer/Horde/Private/Storage/Bundles/V2/PacketWriter.cpp
new file mode 100644
--- /dev/null
+++ b/Engine/Source/Developer/Horde/Private/Storage/Bundles/V2/PacketWriter.cpp
@@ -0,0 +1,274 @@
+// Copyright Epic Games, Inc. All Rights Reserved.
+
+#include "Storage/Bundles/V2/PacketWriter.h"
+#include "Storage/Bundles/Bundle.h"
+#include "Storage/BlobWriter.h"
+#include "Memory/SharedBuffer.h"
+#include "../../../HordePlatform.h"
+
+FPacketWriter::FPacketWriter(FBlobHandle InBundleHandle, FBlobHandle InPacketHandle)
+	: BundleHandle(MoveTemp(InBundleHandle))
+	, PacketHandle(MoveTemp(InPacketHandle))
+	, Buffer(FUniqueBuffer::Alloc(1024).MoveToShared())
+{
+	ImportHandles.SetNum(FPacketImport::Bias);
+	ImportHandles[FPacketImport::CurrentBundleBaseIdx + FPacketImport::Bias] = BundleHandle;
+	ImportHandles[FPacketImport::CurrentPacketBaseIdx + FPacketImport::Bias] = PacketHandle;
+
+	ImportMap.Add(BundleHandle, FPacketImport::CurrentBundleBaseIdx);
+	ImportMap.Add(PacketHandle, FPacketImport::CurrentPacketBaseIdx);
+
+	ImportOffsets.Add(0);
+
+	BufferBasePtr = (uint8*)Buffer.GetData();
+	FBundleSignature(EBundleVersion::LatestV2, 0).Write(BufferBasePtr);
+
+	Length = FBundleSignature::NumBytes + (sizeof(int) * 3);
+	ExportOffsets.Add((uint32)Length);
+
+	NextBlobLength = 0;// sizeof(int32);
+}
+
+FPacketWriter::~FPacketWriter()
+{
+}
+
+size_t FPacketWriter::GetLength() const
+{
+	return Length + NextBlobLength;
+}
+
+void FPacketWriter::AddImport(FBlobHandle Import)
+{
+	int32 ImportIdx = FindOrAddImport(Import);
+	NextBlobImports.Add(ImportIdx);
+}
+
+int FPacketWriter::GetImportCount() const
+{
+	return ImportHandles.Num();
+}
+
+FBlobHandle FPacketWriter::GetImport(int ImportIdx) const
+{
+	return ImportHandles[ImportIdx];
+}
+
+FMutableMemoryView FPacketWriter::GetOutputBuffer(size_t UsedSize, size_t DesiredSize)
+{
+	size_t TotalUsedSize = NextBlobLength + UsedSize;
+	if (TotalUsedSize > 0)
+	{
+		TotalUsedSize += sizeof(int32);
+	}
+
+	FMutableMemoryView Memory = GetOutputBufferInternal(TotalUsedSize, sizeof(int32) + NextBlobLength + DesiredSize);
+	return Memory.Mid(sizeof(int32) + NextBlobLength);
+}
+
+void FPacketWriter::Advance(size_t Size)
+{
+	NextBlobLength += Size;
+}
+
+int FPacketWriter::CompleteBlob(const FBlobType& Type)
+{
+	size_t TypeIdx = FindOrAddType(Type);
+
+	// Ensure the buffer has been allocated
+	GetOutputBuffer(0, 0);
+
+	// Finalize the written export data, inserting the length at the start.
+	*(int*)(BufferBasePtr + Length) = (int)NextBlobLength;
+	Length += sizeof(int) + NextBlobLength;
+
+	// Write the type and import list
+	uint8* Span = (uint8*)GetOutputBufferInternal(0, (NextBlobImports.Num() + 2) * 5).GetData();
+	size_t Offset = 0;
+
+	// Write the type and import list
+	Offset += WriteUnsignedVarInt(Span + Offset, TypeIdx);
+	Offset += WriteUnsignedVarInt(Span + Offset, NextBlobImports.Num());
+
+	for (int32 Import : NextBlobImports)
+	{
+		Offset += WriteUnsignedVarInt(Span + Offset, Import);
+	}
+
+	Length = Align(Length + Offset);
+
+	// Reset the current export
+	NextBlobLength = 0;
+	NextBlobImports.Empty();
+
+	// Write the next export offset to the buffer
+	ExportOffsets.Add((uint32)Length);
+	return ExportOffsets.Num() - 2;
+}
+
+int FPacketWriter::GetExportCount() const
+{
+	return ExportOffsets.Num() - 1;
+}
+
+FBlob FPacketWriter::GetExport(int ExportIdx) const
+{
+	size_t PacketOffset = ExportOffsets[ExportIdx];
+	size_t PacketLength = ExportOffsets[ExportIdx + 1] - PacketOffset;
+
+	FPacketExport Export(FSharedBufferView(Buffer, PacketOffset, PacketLength));
+
+	FBlobType Type = Types[Export.GetTypeIndex()];
+
+	TArray<size_t> ImportIndices;
+	Export.GetImportIndices(ImportIndices);
+
+	TArray<FBlobHandle> Imports;
+	for (size_t ImportIdx : ImportIndices)
+	{
+		Imports.Add(ImportHandles[ImportIdx + FPacketImport::Bias]);
+	}
+
+	return FBlob(Type, Export.GetPayload(), MoveTemp(Imports));
+}
+
+FPacket FPacketWriter::CompletePacket()
+{
+	Length = Align(Length);
+
+	// Write the various lookup tables
+	WriteTypeTable();
+	WriteImportTable();
+	WriteExportTable();
+
+	// Write the final packet length
+	FBundleSignature(EBundleVersion::LatestV2, Length).Write(BufferBasePtr);
+	return FPacket(FSharedBufferView(Buffer, 0, Length));
+}
+
+FMutableMemoryView FPacketWriter::GetOutputBufferInternal(size_t UsedSize, size_t DesiredSize)
+{
+	if (Length + DesiredSize > Buffer.GetSize())
+	{
+		int NewSize = (Length + DesiredSize + 4096 + 16384) & ~16384;
+
+		FSharedBuffer NewBuffer = FUniqueBuffer::Alloc(NewSize).MoveToShared();
+		uint8* NewBufferBasePtr = (uint8*)NewBuffer.GetData();
+		memcpy(NewBufferBasePtr, BufferBasePtr, Length + UsedSize);
+
+		Buffer = MoveTemp(NewBuffer);
+		BufferBasePtr = NewBufferBasePtr;
+	}
+	return FMutableMemoryView(BufferBasePtr + Length, Buffer.GetSize() - Length);
+}
+
+unsigned char* FPacketWriter::GetOutputSpanAndAdvance(size_t Size)
+{
+	FMutableMemoryView Span = GetOutputBufferInternal(0, Size);
+	Length += Size;
+	return (uint8*)Span.GetData();
+}
+
+int32 FPacketWriter::FindOrAddType(FBlobType Type)
+{
+	for (int TypeIdx = 0; TypeIdx < Types.Num(); TypeIdx++)
+	{
+		if (Types[TypeIdx] == Type)
+		{
+			return TypeIdx;
+		}
+	}
+
+	Types.Add(Type);
+	return Types.Num() - 1;
+}
+
+int32 FPacketWriter::FindOrAddImport(const FBlobHandle& Handle)
+{
+	int32* ImportIdxPtr = ImportMap.Find(Handle);
+	if (ImportIdxPtr == nullptr)
+	{
+		FUtf8String Fragment;
+		verify(Handle->TryAppendIdentifier(Fragment));
+
+		FBlobHandle Outer = Handle->GetOuter();
+
+		int BaseIdx;
+		if (Outer)
+		{
+			BaseIdx = FindOrAddImport(Outer);
+		}
+		else
+		{
+			BaseIdx = FPacketImport::InvalidBaseIdx;
+		}
+
+		size_t BaseIdxLen = MeasureUnsignedVarInt(BaseIdx + FPacketImport::Bias);
+		size_t ImportLen = BaseIdxLen + Fragment.Len();
+
+		FMutableMemoryView ImportData = ImportWriter.GetOutputBuffer(0, ImportLen);
+		WriteUnsignedVarIntWithKnownLength(ImportData.GetData(), BaseIdx + FPacketImport::Bias, BaseIdxLen);
+		memcpy((char*)ImportData.GetData() + BaseIdxLen, *Fragment, Fragment.Len());
+		ImportWriter.Advance(ImportLen);
+
+		int32 ImportIdx = ImportOffsets.Num() - 1;
+		ImportHandles.Add(Handle);
+		ImportOffsets.Add(ImportWriter.GetLength());
+		ImportIdxPtr = &ImportMap.Add(Handle, ImportIdx);
+	}
+	return *ImportIdxPtr;
+}
+
+void FPacketWriter::WriteTypeTable()
+{
+	// Write the type table offset in the header
+	*(int*)(BufferBasePtr + 8) = (int)Length;
+
+	// Write the types to the end of the buffer
+	unsigned char* Span = GetOutputSpanAndAdvance(sizeof(int) + Types.Num() * sizeof(FBlobType));
+
+	*(int*)Span = (int)Types.Num();
+	Span += sizeof(int);
+
+	memcpy(Span, Types.GetData(), Types.Num() * sizeof(FBlobType));
+}
+
+void FPacketWriter::WriteImportTable()
+{
+	// Write the the import table offset in the header
+	*(int*)(BufferBasePtr + 12) = (int)Length;
+
+	// Allocate the buffers
+	int ImportHeaderLength = (ImportOffsets.Num() + 1) * sizeof(int);
+	int ImportDataOffset = Length + ImportHeaderLength;
+	unsigned char* Span = GetOutputSpanAndAdvance(Align(ImportHeaderLength + ImportWriter.GetLength()));
+
+	// Write the import header
+	*(int*)Span = ImportOffsets.Num() - 1;
+	Span += sizeof(int);
+
+	for (int ImportOffset : ImportOffsets)
+	{
+		*(int*)Span = ImportDataOffset + ImportOffset;
+		Span += sizeof(int);
+	}
+
+	// Write the import data
+	ImportWriter.CopyTo(Span);
+}
+
+void FPacketWriter::WriteExportTable()
+{
+	// Write the export table offset in the header
+	*(int*)(BufferBasePtr + 16) = (int)Length;
+
+	// Write the exports
+	int* Span = (int*)GetOutputSpanAndAdvance((ExportOffsets.Num() + 1) * sizeof(int));
+	*(Span++) = (int)(ExportOffsets.Num() - 1);
+	memcpy(Span, ExportOffsets.GetData(), ExportOffsets.Num() * sizeof(int));
+}
+
+size_t FPacketWriter::Align(size_t Value)
+{
+	return (Value + 3) & ~3;
+}
diff --git a/Engine/Source/Developer/Horde/Private/Storage/ChunkedBufferWriter.cpp b/Engine/Source/Developer/Horde/Private/Storage/ChunkedBufferWriter.cpp
new file mode 100644
--- /dev/null
+++ b/Engine/Source/Developer/Horde/Private/Storage/ChunkedBufferWriter.cpp
@@ -0,0 +1,179 @@
+// Copyright Epic Games, Inc. All Rights Reserved.
+
+#include "Storage/ChunkedBufferWriter.h"
+#include "../HordePlatform.h"
+
+struct FChunkedBufferWriter::FChunk
+{
+	FSharedBuffer Buffer;
+	size_t Length;
+
+	FChunk(size_t InCapacity)
+		: Buffer(FUniqueBuffer::Alloc(InCapacity).MoveToShared())
+		, Length(0)
+	{
+	}
+
+	FChunk(const FChunk&) = delete;
+
+	FChunk(FChunk&& Other)
+		: Buffer(MoveTemp(Other.Buffer))
+		, Length(Other.Length)
+	{
+	}
+
+	uint8* GetData()
+	{
+		return const_cast<uint8*>(static_cast<const uint8*>(Buffer.GetData()));
+	}
+
+	const uint8* GetData() const
+	{
+		return static_cast<const uint8*>(Buffer.GetData());
+	}
+
+	FMemoryView GetView() const
+	{
+		return FMemoryView(Buffer.GetData(), Length);
+	}
+
+	FMutableMemoryView GetWritableView()
+	{
+		return FMutableMemoryView(GetData() + Length, GetCapacity() - Length);
+	}
+
+	size_t GetCapacity() const
+	{
+		return Buffer.GetSize();
+	}
+
+	size_t GetLength() const
+	{
+		return Length;
+	}
+};
+
+FChunkedBufferWriter::FChunkedBufferWriter(size_t InitialSize)
+{
+	Chunks.Add(FChunk(InitialSize));
+	WrittenLength = 0;
+}
+
+FChunkedBufferWriter::~FChunkedBufferWriter()
+{
+}
+
+void FChunkedBufferWriter::Reset()
+{
+	WrittenLength = 0;
+	Chunks.Empty();
+}
+
+size_t FChunkedBufferWriter::GetLength() const
+{
+	return WrittenLength;
+}
+
+FSharedBufferView FChunkedBufferWriter::Slice(size_t Offset, size_t Length) const
+{
+	size_t EndOffset = Offset + Length;
+	if (EndOffset > WrittenLength)
+	{
+		EndOffset = WrittenLength;
+		Length = EndOffset - Offset;
+	}
+
+	size_t MinChunkIdx = 0;
+	while (MinChunkIdx < Chunks.Num() && Offset > Chunks[MinChunkIdx].GetLength())
+	{
+		Offset -= Chunks[MinChunkIdx].GetLength();
+		MinChunkIdx++;
+	}
+
+	size_t MaxChunkIdx = MinChunkIdx;
+	while (MaxChunkIdx < Chunks.Num() && EndOffset > Chunks[MaxChunkIdx].GetLength())
+	{
+		EndOffset -= Chunks[MaxChunkIdx].GetLength();
+		MaxChunkIdx++;
+	}
+
+	if (MinChunkIdx == MaxChunkIdx)
+	{
+		const FChunk& Chunk = Chunks[MinChunkIdx];
+		return FSharedBufferView(Chunk.Buffer, Offset, Length);
+	}
+
+	FUniqueBuffer Buffer = FUniqueBuffer::Alloc(Length);
+	{
+		size_t SourceOffset = Offset;
+		size_t SourceLength = Length;
+
+		uint8* Output = (uint8*)Buffer.GetData();
+		for (size_t ChunkIdx = MinChunkIdx; ChunkIdx <= MaxChunkIdx; ChunkIdx++)
+		{
+			FMemoryView Source = Chunks[ChunkIdx].GetView();
+			if (SourceOffset > 0)
+			{
+				Source = Source.Mid(SourceOffset);
+			}
+			if (Source.GetSize() > SourceLength)
+			{
+				Source = Source.Left(SourceLength);
+			}
+
+			memcpy(Output, Source.GetData(), Source.GetSize());
+
+			SourceOffset = 0;
+			SourceLength -= Source.GetSize();
+		}
+		check(Output == (const uint8*)Buffer.GetData() + Length);
+	}
+	return Buffer.MoveToShared();
+}
+
+TArray<FMemoryView> FChunkedBufferWriter::GetView() const
+{
+	TArray<FMemoryView> View;
+	for (const FChunk& Chunk : Chunks)
+	{
+		View.Add(Chunk.GetView());
+	}
+	return View;
+}
+
+void FChunkedBufferWriter::CopyTo(void* Buffer) const
+{
+	uint8* End = (uint8*)Buffer;
+	for (const FChunk& Chunk : Chunks)
+	{
+		memcpy(End, Chunk.GetData(), Chunk.GetLength());
+		End += Chunk.GetLength();
+	}
+}
+
+FMutableMemoryView FChunkedBufferWriter::GetOutputBuffer(size_t UsedSize, size_t DesiredSize)
+{
+	FChunk* Chunk = &Chunks.Last();
+	if (Chunk->GetLength() + FMath::Max<size_t>(DesiredSize, 1) > Chunk->GetCapacity())
+	{
+		FChunk NextChunk((DesiredSize + 1024) & ~1023);
+		memcpy(NextChunk.GetData(), Chunk->GetData() + Chunk->GetLength(), UsedSize);
+
+		if (Chunk->GetLength() == 0)
+		{
+			Chunks.Pop();
+		}
+
+		Chunks.Add(MoveTemp(NextChunk));
+		Chunk = &Chunks.Last();
+	}
+	return Chunk->GetWritableView();
+}
+
+void FChunkedBufferWriter::Advance(size_t Size)
+{
+	FChunk& Chunk = Chunks.Last();
+	check(Chunk.Length + Size <= Chunk.GetCapacity());
+	Chunk.Length += Size;
+	WrittenLength += Size;
+}
diff --git a/Engine/Source/Developer/Horde/Private/Storage/Clients/BundleStorageClient.cpp b/Engine/Source/Developer/Horde/Private/Storage/Clients/BundleStorageClient.cpp
new file mode 100644
--- /dev/null
+++ b/Engine/Source/Developer/Horde/Private/Storage/Clients/BundleStorageClient.cpp
@@ -0,0 +1,241 @@
+// Copyright Epic Games, Inc. All Rights Reserved.
+
+#include "Storage/Clients/BundleStorageClient.h"
+#include "Storage/Blob.h"
+#include "Storage/Bundles/V2/Packet.h"
+#include "Storage/Bundles/V2/BundleWriter.h"
+#include "Storage/Bundles/V2/ExportHandle.h"
+#include "Storage/Bundles/V2/PacketHandle.h"
+#include "../../HordePlatform.h"
+
+// ---------------------------------------------------------------------------------
+
+class FBundleStorageClient::FBundleHandleData : public FBlobHandleData, public TSharedFromThis<FBundleStorageClient::FBundleHandleData, ESPMode::ThreadSafe>
+{
+public:
+	FBundleStorageClient* StorageClient;
+	FBlobHandle Inner;
+	mutable bool HasCachedImports;
+	mutable TArray<FBlobHandle> CachedImports;
+
+	FBundleHandleData(FBundleStorageClient* InStorageClient, FBlobHandle InInner)
+		: StorageClient(InStorageClient)
+		, Inner(InInner)
+		, HasCachedImports(false)
+	{ }
+
+	virtual const char* GetType() const override
+	{
+		return "FBundleStorageClient::FBundleHandleData";
+	}
+
+	virtual bool Equals(const FBlobHandleData& Other) const override
+	{
+		return Other.GetType() == GetType() && Inner->Equals(*((const FBundleHandleData&)Other).Inner);
+	}
+
+	virtual uint32 GetHashCode() const
+	{
+		return Inner->GetHashCode();
+	}
+
+	virtual FBlobHandle GetOuter() const override
+	{
+		return Inner->GetOuter();
+	}
+
+	virtual FBlobType ReadType() const override
+	{
+		return FBundle::BlobType;
+	}
+
+	virtual void ReadImports(TArray<FBlobHandle>& OutImports) const override
+	{
+		if (!HasCachedImports)
+		{
+			FBlob BlobData = Read();
+			CachedImports = std::move(BlobData.References);
+			HasCachedImports = true;
+		}
+		OutImports = CachedImports;
+	}
+
+	virtual FBlob Read() const override
+	{
+		FSharedBufferView Data = Inner->ReadBody();
+
+		TArray<FBlobLocator> ImportLocators;
+		ReadImportsFromData(Data.GetView(), ImportLocators);
+
+		TArray<FBlobHandle> ImportHandles;
+		ImportHandles.Reserve(ImportLocators.Num());
+
+		for (const FBlobLocator& ImportLocator : ImportLocators)
+		{
+			ImportHandles.Add(StorageClient->CreateHandle(ImportLocator));
+		}
+
+		return FBlob(FBundle::BlobType, std::move(Data), std::move(ImportHandles));
+	}
+
+	virtual FSharedBufferView ReadBody(size_t Offset, TOptional<size_t> Length) const override
+	{
+		return Inner->ReadBody(Offset, Length);
+	}
+
+	static void ReadImportsFromData(FMemoryView Data, TArray<FBlobLocator>& OutLocators)
+	{
+		FBundleSignature Signature = FBundleSignature::Read(Data.GetData());
+		if (Signature.Version > EBundleVersion::LatestV1 && Signature.Version <= EBundleVersion::LatestV2)
+		{
+			return ReadImportsFromDataV2(Data, OutLocators);
+		}
+		else
+		{
+			FHordePlatform::NotSupported("Unsupported bundle version");
+		}
+	}
+
+	static void ReadImportsFromDataV2(FMemoryView Data, TArray<FBlobLocator>& OutLocators)
+	{
+		while (Data.GetSize() > 0)
+		{
+			FBundleSignature Signature = FBundleSignature::Read(Data.GetData());
+
+			FPacket Packet = FPacket::Decode(Data.Left(Signature.Length));
+			for (size_t Idx = 0; Idx < Packet.GetImportCount(); Idx++)
+			{
+				FPacketImport Import = Packet.GetImport(Idx);
+				if (Import.GetBaseIdx() == FPacketImport::InvalidBaseIdx)
+				{
+					OutLocators.Add(FBlobLocator(Import.GetFragment()));
+				}
+			}
+
+			Data = Data.Mid(Signature.Length);
+		}
+	}
+
+	virtual bool TryAppendIdentifier(FUtf8String& OutBuffer) const override
+	{
+		return Inner->TryAppendIdentifier(OutBuffer);
+	}
+
+	virtual FBlobHandle GetFragmentHandle(const FUtf8StringView& Fragment) const override
+	{
+		TBlobHandle<FBundleHandleData> BundleHandle = const_cast<FBundleStorageClient::FBundleHandleData*>(this)->AsShared();
+		return StorageClient->GetFragmentHandle(BundleHandle, Fragment);
+	}
+};
+
+FBundleStorageClient::FBundleStorageClient(TSharedRef<FKeyValueStorageClient> InInner)
+	: Inner(MoveTemp(InInner))
+{ 
+}
+
+FBundleStorageClient::~FBundleStorageClient()
+{
+}
+
+FBlob FBundleStorageClient::ReadBlob(const FBlobLocator& Locator) const
+{
+	return Inner->ReadBlob(Locator);
+}
+
+FBlobHandle FBundleStorageClient::WriteBlob(const FUtf8StringView& BasePath, const FBlobType& Type, const TArrayView<FMemoryView>& Data, const TArrayView<FBlobHandle>& Imports)
+{
+	return Inner->WriteBlob(BasePath, Type, Data, Imports);
+}
+
+FBlobHandle FBundleStorageClient::CreateHandle(const FBlobLocator& Locator) const
+{
+	FBlobLocator BaseLocator;
+	FUtf8StringView Fragment;
+
+	if (Locator.TryUnwrap(BaseLocator, Fragment))
+	{
+		return GetFragmentHandle(CreateBundleHandle(BaseLocator), Fragment);
+	}
+	else
+	{
+		return CreateBundleHandle(Locator);
+	}
+}
+
+TUniquePtr<FBlobWriter> FBundleStorageClient::CreateWriter(FUtf8String BasePath)
+{
+	return CreateWriter(BasePath, FBundleOptions::Default);
+}
+
+TUniquePtr<FBlobWriter> FBundleStorageClient::CreateWriter(FUtf8String BasePath, const FBundleOptions& Options)
+{
+	if (Options.MaxVersion == EBundleVersion::LatestV2)
+	{
+		return TUniquePtr<FBlobWriter>(new FBundleWriter(StaticCastSharedRef<FKeyValueStorageClient>(AsShared()), BasePath, Options));
+	}
+	else
+	{
+		FHordePlatform::NotSupported("Unsupported bundle version");
+	}
+}
+
+void FBundleStorageClient::AddAlias(const char* Name, FBlobAlias Alias)
+{
+	Inner->AddAlias(Name, MoveTemp(Alias));
+}
+
+void FBundleStorageClient::RemoveAlias(const char* Name, FBlobHandle Handle)
+{
+	Inner->RemoveAlias(Name, MoveTemp(Handle));
+}
+
+void FBundleStorageClient::FindAliases(const char* Name, int MaxResults, TArray<FBlobAlias>& OutAliases)
+{
+	OutAliases.Empty();
+	Inner->FindAliases(Name, MaxResults, OutAliases);
+
+	for (FBlobAlias& Alias : OutAliases)
+	{
+		Alias.Target = CreateHandle(Alias.Target->GetLocator());
+	}
+}
+
+bool FBundleStorageClient::DeleteRef(const FRefName& Name)
+{
+	return Inner->DeleteRef(Name);
+}
+
+FBlobHandle FBundleStorageClient::ReadRef(const FRefName& Name, const FRefCacheTime& CacheTime) const
+{
+	FBlobHandle Target = Inner->ReadRef(Name, CacheTime);
+	if (Target)
+	{
+		Target = CreateHandle(Target->GetLocator());
+	}
+	return Target;
+}
+
+void FBundleStorageClient::WriteRef(const FRefName& Name, const FBlobHandle& Handle, const FRefOptions& Options)
+{
+	Inner->WriteRef(Name, Handle, Options);
+}
+
+TBlobHandle<FBundleStorageClient::FBundleHandleData> FBundleStorageClient::CreateBundleHandle(const FBlobLocator& Locator) const
+{
+	FBlobHandle Handle = Inner->CreateHandle(Locator);
+	return MakeShared<FBundleHandleData>(const_cast<FBundleStorageClient*>(this), Handle);
+}
+
+FBlobHandle FBundleStorageClient::GetFragmentHandle(const FBlobHandle& BundleHandle, const FUtf8StringView& Fragment) const
+{
+	int32 AmpIdx;
+	if (Fragment.FindChar((UTF8CHAR)'&', AmpIdx))
+	{
+		TBlobHandle<FPacketHandleData> PacketHandle = MakeShared<FPacketHandleData>(const_cast<FStorageClient*>(static_cast<const FStorageClient*>(this))->AsShared(), BundleHandle, Fragment.Mid(0, AmpIdx));
+		return MakeShared<FExportHandleData>(PacketHandle, Fragment.Mid(AmpIdx + 1));
+	}
+	else
+	{
+		return MakeShared<FPacketHandleData>(const_cast<FStorageClient*>(static_cast<const FStorageClient*>(this))->AsShared(), BundleHandle, Fragment);
+	}
+}
diff --git a/Engine/Source/Developer/Horde/Private/Storage/Clients/FileStorageClient.cpp b/Engine/Source/Developer/Horde/Private/Storage/Clients/FileStorageClient.cpp
new file mode 100644
--- /dev/null
+++ b/Engine/Source/Developer/Horde/Private/Storage/Clients/FileStorageClient.cpp
@@ -0,0 +1,129 @@
+// Copyright Epic Games, Inc. All Rights Reserved.
+
+#include "Storage/Clients/FileStorageClient.h"
+#include "Storage/Blob.h"
+#include "Storage/BlobType.h"
+#include <iostream>
+#include <sstream>
+#include <fstream>
+#include "../../HordePlatform.h"
+
+FFileStorageClient::FFileStorageClient(std::filesystem::path InRootDir)
+	: RootDir(std::move(InRootDir))
+{
+}
+
+FFileStorageClient::~FFileStorageClient()
+{
+}
+
+FBlobLocator FFileStorageClient::ReadRefFromFile(const std::filesystem::path& File)
+{
+	std::ifstream Stream(File);
+
+	std::stringstream Buffer;
+	Buffer << Stream.rdbuf();
+
+	return FBlobLocator(FUtf8String(Buffer.str().c_str()));
+}
+
+void FFileStorageClient::WriteRefToFile(const std::filesystem::path& File, const FBlobLocator& Locator)
+{
+	std::ofstream Stream(File);
+
+	const FUtf8String& Path = Locator.GetPath();
+	Stream.write((const char*)*Path, Path.Len());
+}
+
+std::filesystem::path FFileStorageClient::GetBlobFile(const FBlobLocator& Locator) const
+{
+	return RootDir / (std::string((const char*)*Locator.GetPath()) + ".blob");
+}
+
+std::filesystem::path FFileStorageClient::GetRefFile(const FRefName& Name) const
+{
+	return RootDir / (std::string((const char*)*Name.GetText()) + ".ref");
+}
+
+void FFileStorageClient::AddAlias(const char* Name, FBlobAlias Alias)
+{
+	FHordePlatform::NotSupported("File storage client does not currently support aliases.");
+}
+
+void FFileStorageClient::RemoveAlias(const char* Name, FBlobHandle Handle)
+{
+	FHordePlatform::NotSupported("File storage client does not currently support aliases.");
+}
+
+void FFileStorageClient::FindAliases(const char* Name, int MaxResults, TArray<FBlobAlias>& OutAliases)
+{
+	FHordePlatform::NotSupported("File storage client does not currently support aliases.");
+}
+
+bool FFileStorageClient::DeleteRef(const FRefName& Name)
+{
+	std::filesystem::path File = GetRefFile(Name);
+	if (std::filesystem::is_regular_file(std::filesystem::status(File)))
+	{
+		std::filesystem::remove(File);
+		return true;
+	}
+	return false;
+}
+
+FBlobHandle FFileStorageClient::ReadRef(const FRefName& Name, const FRefCacheTime& CacheTime) const
+{
+	std::filesystem::path File = GetRefFile(Name);
+	if (std::filesystem::is_regular_file(std::filesystem::status(File)))
+	{
+		return FBlobHandle();
+	}
+
+	return CreateHandle(ReadRefFromFile(File));
+}
+
+void FFileStorageClient::WriteRef(const FRefName& Name, const FBlobHandle& Target, const FRefOptions& Options)
+{
+	FBlobLocator Locator = Target->GetLocator();
+
+	std::filesystem::path File = GetRefFile(Name);
+	std::filesystem::create_directory(File.parent_path());
+
+	std::ofstream Stream(File);
+	Stream << std::string((const char*)*Locator.GetPath());
+}
+
+FBlob FFileStorageClient::ReadBlob(const FBlobLocator& Locator) const
+{
+	std::filesystem::path FileName = GetBlobFile(Locator);
+
+	std::ifstream Stream(FileName, std::ios::binary | std::ios::ate);
+	check(Stream);
+
+	std::streampos Length = Stream.tellg();
+	check(Length != std::streampos(-1));
+	Stream.seekg(0, std::ios::beg);
+
+	FUniqueBuffer Buffer = FUniqueBuffer::Alloc(Length);
+	Stream.read((char*)Buffer.GetData(), Length);
+
+	return FBlob(FBlobType::Leaf, Buffer.MoveToShared(), TArray<FBlobHandle>());
+}
+
+FBlobHandle FFileStorageClient::WriteBlob(const FUtf8StringView& BasePath, const FBlobType& Type, const TArrayView<FMemoryView>& Data, const TArrayView<FBlobHandle>& Imports)
+{
+	FBlobLocator Locator = CreateLocator(BasePath);
+
+	std::filesystem::path FileName = GetBlobFile(Locator);
+	std::filesystem::create_directories(FileName.parent_path());
+
+	std::ofstream Stream(FileName, std::ios::binary);
+	for (FMemoryView& View : Data)
+	{
+		Stream.write((const char*)View.GetData(), View.GetSize());
+	}
+	Stream.close();
+
+	return CreateHandle(Locator);
+}
+
diff --git a/Engine/Source/Developer/Horde/Private/Storage/Clients/KeyValueStorageClient.cpp b/Engine/Source/Developer/Horde/Private/Storage/Clients/KeyValueStorageClient.cpp
new file mode 100644
--- /dev/null
+++ b/Engine/Source/Developer/Horde/Private/Storage/Clients/KeyValueStorageClient.cpp
@@ -0,0 +1,73 @@
+// Copyright Epic Games, Inc. All Rights Reserved.
+
+#include "Storage/Clients/KeyValueStorageClient.h"
+#include "Storage/Blob.h"
+#include "../../HordePlatform.h"
+
+class FKeyValueStorageClient::FHandleData : public FBlobHandleData
+{
+public:
+	FKeyValueStorageClient& StorageClient;
+	const FBlobLocator Locator;
+
+	FHandleData(FKeyValueStorageClient& InStorageClient, const FBlobLocator& InLocator)
+		: StorageClient(InStorageClient)
+		, Locator(InLocator)
+	{
+		check(Locator.GetPath().Len() > 0);
+	}
+
+	virtual const char* GetType() const override
+	{
+		return "FKeyValueStorageClient::FHandleData";
+	}
+
+	virtual uint32 GetHashCode() const override
+	{
+		FHordePlatform::NotImplemented();
+//		return std::hash<std::string>()(Locator.GetPath());
+	}
+
+	virtual bool Equals(const FBlobHandleData& Other) const override
+	{
+		return Other.GetType() == GetType() && ((const FHandleData&)Other).Locator == Locator;
+	}
+
+	virtual FBlobHandle GetOuter() const override
+	{
+		return FBlobHandle();
+	}
+
+	virtual FBlob Read() const override
+	{
+		return StorageClient.ReadBlob(Locator);
+	}
+
+	virtual bool TryAppendIdentifier(FUtf8String& OutIdentifier) const override
+	{
+		OutIdentifier += Locator.GetPath();
+		return true;
+	}
+};
+
+FBlobHandle FKeyValueStorageClient::CreateHandle(const FBlobLocator& Locator) const
+{
+	FBlobLocator BaseLocator;
+	FUtf8StringView Fragment;
+
+	if (Locator.TryUnwrap(BaseLocator, Fragment))
+	{
+		FHordePlatform::NotImplemented();
+//		return new FBlobFragmentHandle(new Handle(this, baseLocator), fragment);
+	}
+	else
+	{
+		return MakeShared<FHandleData>(const_cast<FKeyValueStorageClient&>(*this), Locator);
+	}
+}
+
+TUniquePtr<FBlobWriter> FKeyValueStorageClient::CreateWriter(FUtf8String BasePath)
+{
+	FHordePlatform::NotImplemented();
+}
+
diff --git a/Engine/Source/Developer/Horde/Private/Storage/Clients/MemoryStorageClient.cpp b/Engine/Source/Developer/Horde/Private/Storage/Clients/MemoryStorageClient.cpp
new file mode 100644
--- /dev/null
+++ b/Engine/Source/Developer/Horde/Private/Storage/Clients/MemoryStorageClient.cpp
@@ -0,0 +1,114 @@
+// Copyright Epic Games, Inc. All Rights Reserved.
+
+#include "Storage/Clients/MemoryStorageClient.h"
+#include "Storage/Clients/KeyValueStorageClient.h"
+#include "Storage/Blob.h"
+
+FMemoryStorageClient::FMemoryStorageClient()
+{
+}
+
+FMemoryStorageClient::~FMemoryStorageClient()
+{
+}
+
+void FMemoryStorageClient::AddAlias(const char* Name, FBlobAlias Alias)
+{
+	FBlobLocator Locator = Alias.Target->GetLocator();
+
+	FScopeLock Lock(&CriticalSection);
+
+	TArray<FBlobAlias>& AliasList = Aliases.Emplace(FUtf8String(Name), TArray<FBlobAlias>());
+	AliasList.Add(MoveTemp(Alias));
+}
+
+void FMemoryStorageClient::RemoveAlias(const char* Name, FBlobHandle Handle)
+{
+	FScopeLock Lock(&CriticalSection);
+
+	TArray<FBlobAlias>* AliasList = Aliases.Find(Name);
+	if (AliasList == nullptr)
+	{
+		return;
+	}
+
+	FBlobLocator locator = Handle->GetLocator();
+	AliasList->RemoveAll([&Handle](const FBlobAlias& Alias) -> bool { return Alias.Target->Equals(*Handle); });
+
+	if (AliasList->Num() == 0)
+	{
+		Aliases.Remove(Name);
+	}
+}
+	
+void FMemoryStorageClient::FindAliases(const char* Name, int MaxResults, TArray<FBlobAlias>& OutAliases)
+{
+	FScopeLock Lock(&CriticalSection);
+
+	OutAliases.Empty();
+
+	const TArray<FBlobAlias>* Result = Aliases.Find(Name);
+	if (Result != nullptr)
+	{
+		OutAliases = *Result;
+	}
+}
+
+bool FMemoryStorageClient::DeleteRef(const FRefName& Name)
+{
+	return Refs.Remove(Name) > 0;
+}
+
+FBlobHandle FMemoryStorageClient::ReadRef(const FRefName& Name, const FRefCacheTime& CacheTime) const
+{
+	FScopeLock Lock(&CriticalSection);
+
+	const FBlobLocator* Locator = Refs.Find(Name);
+	if (Locator == nullptr)
+	{
+		return FBlobHandle();
+	}
+	else
+	{
+		return CreateHandle(*Locator);
+	}
+}
+
+void FMemoryStorageClient::WriteRef(const FRefName& Name, const FBlobHandle& Target, const FRefOptions& Options)
+{
+	FScopeLock Lock(&CriticalSection);
+	Refs[Name] = Target->GetLocator();
+}
+
+FBlob FMemoryStorageClient::ReadBlob(const FBlobLocator& Locator) const
+{
+	FScopeLock Lock(&CriticalSection);
+	return Blobs.FindChecked(Locator);
+}
+
+FBlobHandle FMemoryStorageClient::WriteBlob(const FUtf8StringView& BasePath, const FBlobType& Type, const TArrayView<FMemoryView>& Data, const TArrayView<FBlobHandle>& Imports)
+{
+	size_t Length = 0;
+	for (const FMemoryView& View : Data)
+	{
+		Length += View.GetSize();
+	}
+
+	FUniqueBuffer Buffer = FUniqueBuffer::Alloc(Length);
+
+	uint8* Output = (uint8*)Buffer.GetData();
+	for (const FMemoryView& View : Data)
+	{
+		memcpy(Output, View.GetData(), View.GetSize());
+		Output += View.GetSize();
+	}
+
+	FScopeLock Lock(&CriticalSection);
+
+	FBlobLocator Locator = CreateLocator(BasePath);
+
+	FBlob Blob(Type, Buffer.MoveToShared(), TArray<FBlobHandle>(Imports));
+	Blobs.Emplace(Locator, MoveTemp(Blob));
+
+	return CreateHandle(Locator);
+}
diff --git a/Engine/Source/Developer/Horde/Private/Storage/Nodes/ChunkNode.cpp b/Engine/Source/Developer/Horde/Private/Storage/Nodes/ChunkNode.cpp
new file mode 100644
--- /dev/null
+++ b/Engine/Source/Developer/Horde/Private/Storage/Nodes/ChunkNode.cpp
@@ -0,0 +1,167 @@
+// Copyright Epic Games, Inc. All Rights Reserved.
+
+#include "Storage/Nodes/ChunkNode.h"
+#include "Storage/Blob.h"
+
+FChunkNode::~FChunkNode()
+{
+}
+
+// ----------------------------------------------------------------------
+
+FChunkNodeRef::FChunkNodeRef(FBlobHandle InHandle, const FIoHash& InHash)
+	: Handle(MoveTemp(InHandle))
+	, Hash(InHash)
+{ }
+
+FChunkNodeRef::~FChunkNodeRef()
+{ }
+
+// ----------------------------------------------------------------------
+
+const FBlobType FLeafChunkNode::BlobType(FGuid(0xB27AFB68, 0x4A4B9E20, 0x8A78D8A4, 0x39D49840), 1);
+
+FLeafChunkNode::FLeafChunkNode(FSharedBufferView InBuffer)
+	: Buffer(MoveTemp(InBuffer))
+{ }
+
+FLeafChunkNode::~FLeafChunkNode()
+{ }
+
+FLeafChunkNode FLeafChunkNode::Read(FBlob Blob)
+{
+	return FLeafChunkNode(MoveTemp(Blob.Data));
+}
+
+void FLeafChunkNode::Write(FBlobWriter& Writer)
+{
+	void* Target = Writer.GetOutputBuffer(Buffer.GetLength());
+	memcpy(Target, Buffer.GetPointer(), Buffer.GetLength());
+	Writer.Advance(Buffer.GetLength());
+}
+
+// ----------------------------------------------------------------------
+
+const FBlobType FInteriorChunkNode::BlobType(FGuid(0xF4DEDDBC, 0x4C7A70CB, 0x11F04783, 0xB9CDCCAF), 2);
+
+FInteriorChunkNode::FInteriorChunkNode()
+{
+}
+
+FInteriorChunkNode::~FInteriorChunkNode()
+{
+}
+
+FInteriorChunkNode FInteriorChunkNode::Read(FBlob Blob)
+{
+	FInteriorChunkNode Node;
+
+	int32 NumNodes = (int32)(Blob.Data.GetLength() / sizeof(FIoHash));
+	Node.Children.Reserve(NumNodes);
+
+	const FIoHash* Hashes = (const FIoHash*)Blob.Data.GetPointer();
+	for (int32 Idx = 0; Idx < NumNodes; Idx++)
+	{
+		Node.Children.Add(FChunkNodeRef(MoveTemp(Blob.References[Idx]), Hashes[Idx]));
+	}
+
+	return MoveTemp(Node);
+}
+
+void FInteriorChunkNode::Write(FBlobWriter& Writer) const
+{
+	FIoHash* Hashes = (FIoHash*)Writer.GetOutputBuffer(sizeof(FIoHash) * Children.Num());
+	for (int32 Idx = 0; Idx < Children.Num(); Idx++)
+	{
+		Hashes[Idx] = Children[Idx].Hash;
+		Writer.AddImport(Children[Idx].Handle);
+	}
+	Writer.Advance(sizeof(FIoHash) * Children.Num());
+}
+
+// ----------------------------------------------------------------------
+
+struct FChunkNodeReader::FStackEntry
+{
+	FBlob Blob;
+	size_t Position;
+
+	FStackEntry(FBlob InBlob)
+		: Blob(MoveTemp(InBlob))
+		, Position(0)
+	{ }
+};
+
+FChunkNodeReader::FChunkNodeReader(FBlob Blob)
+{
+	Stack.Add(FStackEntry(MoveTemp(Blob)));
+	Advance(0);
+}
+
+FChunkNodeReader::FChunkNodeReader(const FBlobHandle& Handle)
+	: FChunkNodeReader(Handle->Read())
+{
+}
+
+FChunkNodeReader::~FChunkNodeReader()
+{
+}
+
+bool FChunkNodeReader::IsEof() const
+{
+	return Stack.Num() == 0;
+}
+
+FMemoryView FChunkNodeReader::GetBuffer() const
+{
+	if (Stack.Num() == 0)
+	{
+		return FMemoryView();
+	}
+
+	const FStackEntry& StackTop = Stack.Top();
+	check(StackTop.Blob.Type.Guid == FLeafChunkNode::BlobType.Guid);
+
+	FMemoryView View = StackTop.Blob.Data.GetView();
+	return View.Mid(StackTop.Position);
+}
+
+void FChunkNodeReader::Advance(int32 Size)
+{
+	while (Stack.Num() > 0)
+	{
+		FStackEntry& StackTop = Stack.Top();
+		if (StackTop.Blob.Type.Guid == FLeafChunkNode::BlobType.Guid)
+		{
+			FMemoryView BlobData = StackTop.Blob.Data.GetView();
+
+			size_t ChunkSize = FMath::Min<size_t>(Size, BlobData.GetSize() - StackTop.Position);
+			StackTop.Position += ChunkSize;
+			Size -= ChunkSize;
+
+			if (StackTop.Position < BlobData.GetSize())
+			{
+				break;
+			}
+
+			Stack.Pop();
+		}
+		else if (StackTop.Blob.Type.Guid == FInteriorChunkNode::BlobType.Guid)
+		{
+			FBlobHandle ChildHandle = StackTop.Blob.References[StackTop.Position];
+
+			StackTop.Position++;
+			if (StackTop.Position == StackTop.Blob.References.Num())
+			{
+				Stack.Pop();
+			}
+
+			Stack.Add(FStackEntry(ChildHandle->Read()));
+		}
+		else
+		{
+			// Invalid blob type
+			check(false);
+		}
+	}
+}
diff --git a/Engine/Source/Developer/Horde/Private/Storage/Nodes/DirectoryEntry.cpp b/Engine/Source/Developer/Horde/Private/Storage/Nodes/DirectoryEntry.cpp
new file mode 100644
--- /dev/null
+++ b/Engine/Source/Developer/Horde/Private/Storage/Nodes/DirectoryEntry.cpp
@@ -0,0 +1,38 @@
+// Copyright Epic Games, Inc. All Rights Reserved.
+
+#include "Storage/Nodes/DirectoryEntry.h"
+#include "Storage/BlobReader.h"
+#include "Storage/BlobWriter.h"
+
+FDirectoryEntry::FDirectoryEntry(FBlobHandle InTarget, const FIoHash& InTargetHash, FUtf8String InName, int64 InLength)
+	: Target(MoveTemp(InTarget))
+	, TargetHash(InTargetHash)
+	, Length(InLength)
+	, Name(MoveTemp(InName))
+{
+}
+
+FDirectoryEntry::~FDirectoryEntry()
+{
+}
+
+FDirectoryEntry FDirectoryEntry::Read(FBlobReader& Reader)
+{
+	FBlobHandle Target = Reader.ReadImport();
+	FIoHash TargetHash = ReadIoHash(Reader);
+
+	int64 Length = (int64)ReadUnsignedVarInt(Reader);
+	FUtf8String Name = ReadString(Reader);
+
+	return FDirectoryEntry(MoveTemp(Target), TargetHash, MoveTemp(Name), Length);
+}
+
+void FDirectoryEntry::Write(FBlobWriter& Writer) const
+{
+	Writer.AddImport(Target);
+	WriteIoHash(Writer, TargetHash);
+
+	WriteUnsignedVarInt(Writer, Length);
+	WriteString(Writer, Name);
+}
+
diff --git a/Engine/Source/Developer/Horde/Private/Storage/Nodes/DirectoryNode.cpp b/Engine/Source/Developer/Horde/Private/Storage/Nodes/DirectoryNode.cpp
new file mode 100644
--- /dev/null
+++ b/Engine/Source/Developer/Horde/Private/Storage/Nodes/DirectoryNode.cpp
@@ -0,0 +1,64 @@
+// Copyright Epic Games, Inc. All Rights Reserved.
+
+#include "Storage/Nodes/DirectoryNode.h"
+#include "Storage/Blob.h"
+#include "Storage/BlobReader.h"
+#include "Storage/BlobWriter.h"
+#include "../../HordePlatform.h"
+
+const FBlobType FDirectoryNode::BlobType(FGuid(0x0714EC11, 0x4D07291A, 0x8AE77F86, 0x799980D6), 1);
+
+FDirectoryNode::FDirectoryNode(EDirectoryFlags InFlags)
+	: Flags(InFlags)
+{
+}
+
+FDirectoryNode::~FDirectoryNode()
+{
+}
+
+FDirectoryNode FDirectoryNode::Read(const FBlob& Blob)
+{
+	FBlobReader Reader(Blob);
+
+	EDirectoryFlags Flags = (EDirectoryFlags)ReadUnsignedVarInt(Reader);
+	FDirectoryNode Directory(Flags);
+
+	int32 FileCount = (int32)ReadUnsignedVarInt(Reader);
+	for (int32 Idx = 0; Idx < FileCount; Idx++)
+	{
+		FFileEntry Entry = FFileEntry::Read(Reader);
+		FUtf8String Name = Entry.Name;
+		Directory.NameToFile.Emplace(MoveTemp(Name), MoveTemp(Entry));
+	}
+
+	int32 DirectoryCount = (int32)ReadUnsignedVarInt(Reader);
+	for (int32 Idx = 0; Idx < DirectoryCount; Idx++)
+	{
+		FDirectoryEntry Entry = FDirectoryEntry::Read(Reader);
+		FUtf8String Name = Entry.Name;
+		Directory.NameToDirectory.Emplace(MoveTemp(Name), MoveTemp(Entry));
+	}
+
+	return Directory;
+}
+
+FBlobHandle FDirectoryNode::Write(FBlobWriter& Writer) const
+{
+	WriteUnsignedVarInt(Writer, (int64)Flags);
+
+	WriteUnsignedVarInt(Writer, NameToFile.Num());
+	for(TMap<FUtf8String, FFileEntry>::TConstIterator Iter(NameToFile); Iter; ++Iter)
+	{
+		Iter.Value().Write(Writer);
+	}
+
+	WriteUnsignedVarInt(Writer, NameToDirectory.Num());
+	for (TMap<FUtf8String, FDirectoryEntry>::TConstIterator Iter(NameToDirectory); Iter; ++Iter)
+	{
+		Iter.Value().Write(Writer);
+	}
+
+	return Writer.CompleteBlob(BlobType);
+}
+
diff --git a/Engine/Source/Developer/Horde/Private/Storage/Nodes/FileEntry.cpp b/Engine/Source/Developer/Horde/Private/Storage/Nodes/FileEntry.cpp
new file mode 100644
--- /dev/null
+++ b/Engine/Source/Developer/Horde/Private/Storage/Nodes/FileEntry.cpp
@@ -0,0 +1,61 @@
+// Copyright Epic Games, Inc. All Rights Reserved.
+
+#include "Storage/Nodes/FileEntry.h"
+#include "IO/IoHash.h"
+#include "Storage/BlobReader.h"
+#include "Storage/BlobWriter.h"
+#include "../../HordePlatform.h"
+
+FFileEntry::FFileEntry(FBlobHandle InTarget, const FIoHash& InTargetHash, FUtf8String InName, EFileEntryFlags InFlags, int64 InLength, const FIoHash& InHash, FSharedBufferView InCustomData)
+	: Target(MoveTemp(InTarget))
+	, TargetHash(InTargetHash)
+	, Name(MoveTemp(InName))
+	, Flags(InFlags)
+	, Length(InLength)
+	, Hash(InHash)
+	, CustomData(MoveTemp(InCustomData))
+{ }
+
+FFileEntry::~FFileEntry()
+{ }
+
+FFileEntry FFileEntry::Read(FBlobReader& Reader)
+{
+	FBlobHandle Target = Reader.ReadImport();
+	FIoHash TargetHash = ReadIoHash(Reader);
+
+	FUtf8String Name = ReadString(Reader);
+	EFileEntryFlags Flags = (EFileEntryFlags)ReadUnsignedVarInt(Reader);
+	int64 Length = (int64)ReadUnsignedVarInt(Reader);
+	FIoHash Hash = ReadIoHash(Reader);
+
+	FSharedBufferView CustomData;
+	if (EnumHasAnyFlags(Flags, EFileEntryFlags::HasCustomData))
+	{
+		FHordePlatform::NotImplemented();
+//		CustomData = ReadVariableLengthBytes(CustomData);
+//		Flags &= ~EFileEntryFlags::HasCustomData;
+	}
+
+	return FFileEntry(MoveTemp(Target), TargetHash, MoveTemp(Name), Flags, Length, Hash, MoveTemp(CustomData));
+}
+
+void FFileEntry::Write(FBlobWriter& Writer) const
+{
+	Writer.AddImport(Target);
+	WriteIoHash(Writer, TargetHash);
+
+	EFileEntryFlags WriteFlags = Flags; // TODO: (CustomData.Length > 0) ? (Flags | EFileEntryFlags::HasCustomData) : (Flags & ~EFileEntryFlags::HasCustomData);
+
+	WriteString(Writer, Name);
+	WriteUnsignedVarInt(Writer, (int)WriteFlags);
+	WriteUnsignedVarInt(Writer, Length);
+	WriteIoHash(Writer, Hash);
+
+	if (EnumHasAnyFlags(WriteFlags, EFileEntryFlags::HasCustomData))
+	{
+		FHordePlatform::NotImplemented();
+//		writer.WriteVariableLengthBytes(CustomData.Span);
+	}
+}
+
diff --git a/Engine/Source/Developer/Horde/Private/Storage/RefName.cpp b/Engine/Source/Developer/Horde/Private/Storage/RefName.cpp
new file mode 100644
--- /dev/null
+++ b/Engine/Source/Developer/Horde/Private/Storage/RefName.cpp
@@ -0,0 +1,31 @@
+// Copyright Epic Games, Inc. All Rights Reserved.
+
+#include "Storage/RefName.h"
+
+FRefName::FRefName(FUtf8String InText)
+	: Text(MoveTemp(InText))
+{
+}
+FRefName::~FRefName()
+{
+}
+
+const FUtf8String& FRefName::GetText() const
+{
+	return Text;
+}
+
+bool FRefName::operator==(const FRefName& Other) const
+{
+	return Text.Equals(Other.Text, ESearchCase::CaseSensitive);
+}
+
+bool FRefName::operator!=(const FRefName& Other) const
+{
+	return !(*this == Other);
+}
+
+uint32 GetTypeHash(const FRefName& RefName)
+{
+	return GetTypeHash(RefName.Text);
+}
diff --git a/Engine/Source/Developer/Horde/Private/Storage/RefOptions.cpp b/Engine/Source/Developer/Horde/Private/Storage/RefOptions.cpp
new file mode 100644
--- /dev/null
+++ b/Engine/Source/Developer/Horde/Private/Storage/RefOptions.cpp
@@ -0,0 +1,10 @@
+// Copyright Epic Games, Inc. All Rights Reserved.
+
+#include "Storage/RefOptions.h"
+
+const FRefOptions FRefOptions::Default;
+
+FRefOptions::FRefOptions()
+{
+}
+
diff --git a/Engine/Source/Developer/Horde/Private/Storage/StorageClient.cpp b/Engine/Source/Developer/Horde/Private/Storage/StorageClient.cpp
new file mode 100644
--- /dev/null
+++ b/Engine/Source/Developer/Horde/Private/Storage/StorageClient.cpp
@@ -0,0 +1,78 @@
+// Copyright Epic Games, Inc. All Rights Reserved.
+
+#include "Storage/StorageClient.h"
+#include "Storage/RefOptions.h"
+#include "../HordePlatform.h"
+
+FRefCacheTime::FRefCacheTime()
+{
+}
+
+FRefCacheTime::FRefCacheTime(FTimespan InMaxAge)
+	: Time(FDateTime::UtcNow() - InMaxAge)
+{
+}
+
+FRefCacheTime::FRefCacheTime(FDateTime InTime)
+	: Time(InTime)
+{
+}
+
+// --------------------------------------------------------------------------------------------------------
+
+FStorageClient::~FStorageClient()
+{
+}
+
+void FStorageClient::WriteRef(const FRefName& Name, const FBlobHandle& Handle)
+{
+	WriteRef(Name, Handle, FRefOptions::Default);
+}
+
+FBlobLocator FStorageClient::CreateLocator(const FUtf8StringView& BasePath)
+{
+	FUtf8String Path(BasePath);
+	if (Path.Len() > 0 && !Path.EndsWith("/", ESearchCase::CaseSensitive))
+	{
+		Path += "/";
+	}
+
+	char Name[128];
+	FHordePlatform::CreateUniqueIdentifier(Name, sizeof(Name) / sizeof(Name[0]));
+	Path += Name;
+
+	return FBlobLocator(MoveTemp(Path));
+}
+
+// --------------------------------------------------------------------------------------------------------
+
+FBlobAlias::FBlobAlias(FBlobHandle InTarget, int InRank, FSharedBufferView InData)
+	: Target(MoveTemp(InTarget))
+	, Rank(InRank)
+	, Data(std::move(InData))
+{
+}
+
+FBlobAlias::~FBlobAlias()
+{
+}
+
+// --------------------------------------------------------------------------------------------------------
+
+FAliasInfo::FAliasInfo(FUtf8String InName, int InRank)
+	: Name(MoveTemp(InName))
+	, Rank(InRank)
+{
+}
+
+FAliasInfo::FAliasInfo(FUtf8String InName, int InRank, FSharedBufferView InData)
+	: Name(MoveTemp(InName))
+	, Rank(InRank)
+	, Data(MoveTemp(InData))
+{
+}
+
+FAliasInfo::~FAliasInfo()
+{
+}
+
diff --git a/Engine/Source/Developer/Horde/Public/SharedBufferView.h b/Engine/Source/Developer/Horde/Public/SharedBufferView.h
new file mode 100644
--- /dev/null
+++ b/Engine/Source/Developer/Horde/Public/SharedBufferView.h
@@ -0,0 +1,29 @@
+#pragma once
+
+#include "Memory/SharedBuffer.h"
+
+//
+// View into a buffer with shared ownership
+//
+class FSharedBufferView final
+{
+public:
+	FSharedBufferView();
+	FSharedBufferView(FSharedBuffer InBuffer);
+	FSharedBufferView(FSharedBuffer InBuffer, const FMemoryView& InView);
+	FSharedBufferView(FSharedBuffer InBuffer, size_t InOffset, size_t InLength);
+	~FSharedBufferView();
+
+	static FSharedBufferView Copy(const FMemoryView& View);
+
+	FSharedBufferView Slice(uint64 Offset) const;
+	FSharedBufferView Slice(uint64 Offset, uint64 Length) const;
+
+	const unsigned char* GetPointer() const;
+	size_t GetLength() const;
+	FMemoryView GetView() const;
+
+private:
+	FSharedBuffer Buffer;
+	FMemoryView View;
+};
diff --git a/Engine/Source/Developer/Horde/Public/Storage/Blob.h b/Engine/Source/Developer/Horde/Public/Storage/Blob.h
new file mode 100644
--- /dev/null
+++ b/Engine/Source/Developer/Horde/Public/Storage/Blob.h
@@ -0,0 +1,26 @@
+// Copyright Epic Games, Inc. All Rights Reserved.
+
+#pragma once
+
+#include "CoreMinimal.h"
+#include "BlobHandle.h"
+#include "BlobType.h"
+#include "../SharedBufferView.h"
+
+/**
+ * Describes a blob of data in the storage system
+ */
+struct FBlob
+{
+	/** Type of the blob. */
+	FBlobType Type;
+
+	/** Data for the blob. */
+	FSharedBufferView Data;
+
+	/** References to other blobs. */
+	TArray<FBlobHandle> References;
+
+	FBlob(const FBlobType& InType, FSharedBufferView InData, TArray<FBlobHandle> InReferences);
+	~FBlob();
+};
diff --git a/Engine/Source/Developer/Horde/Public/Storage/BlobHandle.h b/Engine/Source/Developer/Horde/Public/Storage/BlobHandle.h
new file mode 100644
--- /dev/null
+++ b/Engine/Source/Developer/Horde/Public/Storage/BlobHandle.h
@@ -0,0 +1,112 @@
+// Copyright Epic Games, Inc. All Rights Reserved.
+
+#pragma once
+
+#include "CoreMinimal.h"
+#include "SharedBufferView.h"
+#include "BlobType.h"
+
+struct FBlob;
+class FBlobHandleData;
+class FBlobLocator;
+
+/**
+ * Handle to a blob. Can be used to reference blobs that have not been flushed yet. Implemented as a shared pointer to a FBlobHandleData object.
+ */
+class FBlobHandle : public TSharedPtr<FBlobHandleData>
+{
+public:
+	using Super = TSharedPtr<FBlobHandleData>;
+	using Super::operator bool;
+	using Super::operator->;
+	using Super::operator*;
+
+	FBlobHandle()
+	{ }
+
+	FBlobHandle(TSharedPtr<FBlobHandleData>&& Other)
+		: TSharedPtr<FBlobHandleData>(Other)
+	{ }
+
+	FBlobHandle(TSharedRef<FBlobHandleData> Value)
+		: TSharedPtr<FBlobHandleData>(MoveTemp(Value))
+	{ }
+
+	template<typename T>
+	FBlobHandle(TSharedPtr<T> Value)
+		: TSharedPtr<FBlobHandleData>(StaticCastSharedPtr<FBlobHandleData, T>(MoveTemp(Value)))
+	{ }
+
+	template<typename T>
+	FBlobHandle(TSharedRef<T> Value)
+		: TSharedPtr<FBlobHandleData>(StaticCastSharedPtr<FBlobHandleData, T>(Value.ToSharedPtr()))
+	{ }
+
+	bool operator==(const FBlobHandle& Other) const;
+	bool operator!=(const FBlobHandle& Other) const;
+
+	friend uint32 GetTypeHash(const FBlobHandle& Handle);
+};
+
+/**
+ * Handle to a blob with a particular state representation
+ */
+template<typename T>
+class TBlobHandle : public FBlobHandle
+{
+public:
+	TBlobHandle()
+	{ }
+
+	TBlobHandle(TSharedRef<T> Other)
+		: FBlobHandle(MoveTemp(Other))
+	{ }
+
+	T* operator->() { return static_cast<T*>(FBlobHandle::operator->()); }
+	const T* operator->() const { return static_cast<const T*>(FBlobHandle::operator->()); }
+};
+
+/**
+ * Base interface for a blob.
+ */
+class FBlobHandleData
+{
+public:
+	virtual ~FBlobHandleData();
+
+	/** Gets the type name of this derived blob handle instance. */
+	virtual const char* GetType() const = 0;
+
+	/** Determines if this blob equals the other. */
+	virtual bool Equals(const FBlobHandleData& Other) const = 0;
+
+	/** Gets a hash for this blob. */
+	virtual uint32 GetHashCode() const = 0;
+
+	/** For a blob nested within another blob, gets a handle to the containing blob (eg. For a bundle node, will return the packet. For a bundle packet, will return the bundle. For a bundle or other non-nested blob, returns null.) */
+	virtual FBlobHandle GetOuter() const = 0;
+
+	/** Reads the blob's data. */
+	virtual FBlob Read() const = 0;
+
+	/** Gets the type of this blob. */
+	virtual FBlobType ReadType() const;
+
+	/** Gets the outward references from this blob. */
+	virtual void ReadImports(TArray<FBlobHandle>& OutImports) const;
+
+	/** Reads part of the blob, and returns a handle that can be used to access the data. */
+	virtual FSharedBufferView ReadBody() const;
+
+	/** Reads part of the blob, and returns a handle that can be used to access the data. */
+	virtual FSharedBufferView ReadBody(size_t Offset, TOptional<size_t> Length = TOptional<size_t>()) const;
+
+	/** Gets a locator for this blob, by calling TryAppendLocator recursively. */
+	FBlobLocator GetLocator() const;
+
+	/** Gets an identifier for this blob, relative to its outer. */
+	virtual bool TryAppendIdentifier(FUtf8String& OutBuffer) const = 0;
+
+	/** Gets a handle to a nested blob object. */
+	virtual FBlobHandle GetFragmentHandle(const FUtf8StringView& Fragment) const;
+};
diff --git a/Engine/Source/Developer/Horde/Public/Storage/BlobLocator.h b/Engine/Source/Developer/Horde/Public/Storage/BlobLocator.h
new file mode 100644
--- /dev/null
+++ b/Engine/Source/Developer/Horde/Public/Storage/BlobLocator.h
@@ -0,0 +1,44 @@
+// Copyright Epic Games, Inc. All Rights Reserved.
+
+#pragma once
+
+#include "CoreMinimal.h"
+#include "Containers/Utf8String.h"
+
+/**
+ * Identifies the location of a blob. Meaning of this string is implementation defined.
+ */
+class FBlobLocator
+{
+public:
+	FBlobLocator();
+	FBlobLocator(FUtf8String InPath);
+	FBlobLocator(const FBlobLocator& InBaseLocator, const FUtf8StringView& InFragment);
+	explicit FBlobLocator(const FUtf8StringView& InPath);
+
+	/** Tests whether this locator is valid. */
+	bool IsValid() const;
+
+	/** Gets the path for this locator. */
+	const FUtf8String& GetPath() const;
+
+	/** Gets the base locator for this blob. */
+	FBlobLocator GetBaseLocator() const;
+
+	/** Gets the path for this locator. */
+	FUtf8StringView GetFragment() const;
+
+	/** Determines if this locator can be unwrapped into an outer locator/fragment pair. */
+	bool CanUnwrap() const;
+
+	/** Split this locator into a locator and fragment. */
+	bool TryUnwrap(FBlobLocator& OutLocator, FUtf8StringView& OutFragment) const;
+
+	bool operator==(const FBlobLocator& Other) const;
+	bool operator!=(const FBlobLocator& Other) const;
+
+	friend uint32 GetTypeHash(const FBlobLocator& Locator);
+
+private:
+	FUtf8String Path;
+};
\ No newline at end of file
diff --git a/Engine/Source/Developer/Horde/Public/Storage/BlobReader.h b/Engine/Source/Developer/Horde/Public/Storage/BlobReader.h
new file mode 100644
--- /dev/null
+++ b/Engine/Source/Developer/Horde/Public/Storage/BlobReader.h
@@ -0,0 +1,47 @@
+// Copyright Epic Games, Inc. All Rights Reserved.
+
+#pragma once
+
+#include "CoreMinimal.h"
+#include "Blob.h"
+#include "BlobHandle.h"
+#include "BlobLocator.h"
+#include "RefName.h"
+
+struct FIoHash;
+
+/**
+ * Reads data from a blob object
+ */
+class FBlobReader
+{
+public:
+	FBlobReader(const FBlob& InBlob);
+	FBlobReader(const FMemoryView& InBuffer, const TArray<FBlobHandle>& InImports);
+
+	/** Gets a pointer to the remaining memory. */
+	const unsigned char* GetBuffer() const;
+
+	/** Gets the remaining memory to read from. */
+	FMemoryView GetView() const;
+
+	/** Advance the current read position. */
+	void Advance(size_t Size);
+
+	/** Reads the next import. */
+	FBlobHandle ReadImport();
+
+private:
+	FMemoryView Buffer;
+	const TArray<FBlobHandle>& Imports;
+	int32 NextImportIdx;
+};
+
+// ------------------------------------------------------------------------
+
+int ReadInt32(FBlobReader& Reader);
+FIoHash ReadIoHash(FBlobReader& Reader);
+FMemoryView ReadFixedLengthBytes(FBlobReader& Reader, size_t Length);
+size_t ReadUnsignedVarInt(FBlobReader& Reader);
+FUtf8String ReadString(FBlobReader& Reader);
+FMemoryView ReadStringSpan(FBlobReader& Reader);
diff --git a/Engine/Source/Developer/Horde/Public/Storage/BlobType.h b/Engine/Source/Developer/Horde/Public/Storage/BlobType.h
new file mode 100644
--- /dev/null
+++ b/Engine/Source/Developer/Horde/Public/Storage/BlobType.h
@@ -0,0 +1,24 @@
+// Copyright Epic Games, Inc. All Rights Reserved.
+
+#pragma once
+
+#include "CoreMinimal.h"
+
+/**
+ * Identifies the type of a blob
+ */
+struct FBlobType
+{
+	static const FBlobType Leaf;
+
+	/** Nominal identifier for the type. */
+	FGuid Guid;
+
+	/** Version number for the serializer. */
+	int Version;
+
+	FBlobType(const FGuid& InGuid, int InVersion);
+
+	bool operator==(const FBlobType& Other) const;
+	bool operator!=(const FBlobType& Other) const;
+};
diff --git a/Engine/Source/Developer/Horde/Public/Storage/BlobWriter.h b/Engine/Source/Developer/Horde/Public/Storage/BlobWriter.h
new file mode 100644
--- /dev/null
+++ b/Engine/Source/Developer/Horde/Public/Storage/BlobWriter.h
@@ -0,0 +1,68 @@
+// Copyright Epic Games, Inc. All Rights Reserved.
+
+#pragma once
+
+#include "CoreMinimal.h"
+#include "Blob.h"
+#include "BlobHandle.h"
+#include "BlobLocator.h"
+#include "RefName.h"
+#include "../SharedBufferView.h"
+
+struct FAliasInfo;
+struct FRefOptions;
+
+/**
+ * Interface for writing data to the storage system.
+ */
+class FBlobWriter
+{
+public:
+	virtual ~FBlobWriter();
+
+	/** Create another writer instance, allowing multiple threads to write in parallel. */
+	virtual TUniquePtr<FBlobWriter> Fork() = 0;
+
+	/** Adds an alias for the current blob. */
+	virtual void AddAlias(const FAliasInfo& AliasInfo) = 0;
+
+	/** Adds a reference to another blob. */
+	virtual void AddImport(FBlobHandle Target) = 0;
+
+	/** Adds a named reference to the blob being built. */
+	void AddRef(const FRefName& RefName);
+	virtual void AddRef(const FRefName& RefName, const FRefOptions& Options) = 0;
+
+	/** Gets a block of memory of the given size. */
+	void* GetOutputBuffer(size_t Size);
+
+	/** Gets a block of memory, at least the given size. */
+	virtual FMutableMemoryView GetOutputBufferAsSpan(size_t UsedSize, size_t DesiredSize) = 0;
+
+	/** Advance the current write position. */
+	virtual void Advance(size_t Size) = 0;
+
+	/** Finish writing a blob that has been written into the output buffer. */
+	virtual FBlobHandle CompleteBlob(const FBlobType& InType) = 0;
+
+	/** Flush any pending nodes to storage. */
+	virtual void Flush() = 0;
+};
+
+// ------------------------------------------------------------------------------------------
+
+void WriteIoHash(FBlobWriter& Writer, const struct FIoHash& Hash);
+
+void WriteFixedLengthBytes(FBlobWriter& Writer, const void* Data, size_t Length);
+void WriteFixedLengthBytes(FBlobWriter& Writer, const FMemoryView& View);
+
+size_t MeasureUnsignedVarInt(size_t Value);
+size_t WriteUnsignedVarInt(void* Buffer, size_t Value);
+void WriteUnsignedVarIntWithKnownLength(void* Buffer, size_t Value, size_t NumBytes);
+void WriteUnsignedVarInt(FBlobWriter& Writer, size_t Value);
+
+size_t MeasureString(const char* Text);
+size_t MeasureString(const FUtf8StringView& Text);
+void WriteString(FBlobWriter& Writer, const char* Text);
+void WriteString(FBlobWriter& Writer, const FUtf8StringView& Text);
+void WriteString(FBlobWriter& Writer, const FUtf8String& Text);
diff --git a/Engine/Source/Developer/Horde/Public/Storage/Bundles/Bundle.h b/Engine/Source/Developer/Horde/Public/Storage/Bundles/Bundle.h
new file mode 100644
--- /dev/null
+++ b/Engine/Source/Developer/Horde/Public/Storage/Bundles/Bundle.h
@@ -0,0 +1,66 @@
+// Copyright Epic Games, Inc. All Rights Reserved.
+
+#pragma once
+
+#include "Storage/Blob.h"
+
+/** Bundle version number */
+enum class EBundleVersion : unsigned char
+{
+	/** Initial version number. */
+	Initial = 0,
+
+	/** Added the BundleExport.Alias property. */
+	ExportAliases = 1,
+
+	/** Back out change to include aliases. Will likely do this through an API rather than baked into the data. */
+	RemoveAliases = 2,
+
+	/** Use data structures which support in-place reading and writing. */
+	InPlace = 3,
+
+	/** Add import hashes to imported nodes. */
+	ImportHashes = 4,
+
+	/** Structure bundles as a sequence of self-contained packets (uses V2 code). */
+	PacketSequence = 5,
+
+	/** Last item in the enum. Used for Latest. */
+	LatestPlusOne,
+
+	/** The current version number. */
+	Latest = (int)LatestPlusOne - 1,
+
+	/** Last version using the V1 pipeline. */
+	LatestV1 = ImportHashes,
+
+	/** Last version using the V2 pipeline. */
+	LatestV2 = Latest,
+};
+
+/**
+ * General bundle properties
+ */
+struct FBundle
+{
+	static const FBlobType BlobType;
+};
+
+/**
+ * Signature for a bundle
+ */
+struct FBundleSignature
+{
+	/** Number of bytes in a serialized signature. */
+	static const int NumBytes = 8;
+
+	/** Version number for the following file data. */
+	EBundleVersion Version;
+
+	/** Length of the packet, including the length of the signature itself. */
+	size_t Length;
+
+	FBundleSignature(EBundleVersion InVersion, size_t InLength);
+	static FBundleSignature Read(const void* Data);
+	void Write(void* Data);
+};
diff --git a/Engine/Source/Developer/Horde/Public/Storage/Bundles/BundleCompression.h b/Engine/Source/Developer/Horde/Public/Storage/Bundles/BundleCompression.h
new file mode 100644
--- /dev/null
+++ b/Engine/Source/Developer/Horde/Public/Storage/Bundles/BundleCompression.h
@@ -0,0 +1,40 @@
+// Copyright Epic Games, Inc. All Rights Reserved.
+
+#pragma once
+
+#include "CoreMinimal.h"
+#include "Memory/MemoryView.h"
+
+/** Indicates the compression format in a bundle */
+enum EBundleCompressionFormat : unsigned char
+{
+	/** Packets are uncompressed. */
+	None = 0,
+
+	/** LZ4 compression. */
+	LZ4 = 1,
+
+	/** Gzip compression. */
+	Gzip = 2,
+
+	/** Oodle compression (Kraken). */
+	Oodle = 3,
+
+	/** Brotli compression. */
+	Brotli = 4,
+};
+
+/*
+ * Utility methods for compressing bundles
+ */
+struct FBundleCompression
+{
+	/** Gets the maximum size of the buffer required to compress the given data. */
+	static size_t GetMaxSize(EBundleCompressionFormat Format, const FMemoryView& Input);
+
+	/** Compress a data packet. */
+	static size_t Compress(EBundleCompressionFormat Format, const FMemoryView& Input, const FMutableMemoryView& Output);
+
+	/** Decompress a packet of data. */
+	static void Decompress(EBundleCompressionFormat Format, const FMemoryView& Input, const FMutableMemoryView& Output);
+};
diff --git a/Engine/Source/Developer/Horde/Public/Storage/Bundles/BundleOptions.h b/Engine/Source/Developer/Horde/Public/Storage/Bundles/BundleOptions.h
new file mode 100644
--- /dev/null
+++ b/Engine/Source/Developer/Horde/Public/Storage/Bundles/BundleOptions.h
@@ -0,0 +1,30 @@
+// Copyright Epic Games, Inc. All Rights Reserved.
+
+#pragma once
+
+#include "Bundle.h"
+#include "BundleCompression.h"
+
+/*
+ * Options for configuring a bundle serializer
+ */
+struct FBundleOptions
+{
+	/** Default options value. */
+	static const FBundleOptions Default;
+
+	/** Maximum version number of bundles to write. */
+	EBundleVersion MaxVersion = EBundleVersion::LatestV2;
+
+	/** Maximum payload size fo a blob. */
+	int MaxBlobSize = 10 * 1024 * 1024;
+
+	/** Compression format to use. */
+	EBundleCompressionFormat CompressionFormat = EBundleCompressionFormat::LZ4;
+
+	/** Minimum size of a block to be compressed. */
+	int MinCompressionPacketSize = 16 * 1024;
+
+	FBundleOptions();
+};
+
diff --git a/Engine/Source/Developer/Horde/Public/Storage/Bundles/V2/BundleWriter.h b/Engine/Source/Developer/Horde/Public/Storage/Bundles/V2/BundleWriter.h
new file mode 100644
--- /dev/null
+++ b/Engine/Source/Developer/Horde/Public/Storage/Bundles/V2/BundleWriter.h
@@ -0,0 +1,56 @@
+// Copyright Epic Games, Inc. All Rights Reserved.
+
+#pragma once
+
+#include "CoreMinimal.h"
+#include "Containers/Utf8String.h"
+#include "Storage/StorageClient.h"
+#include "Storage/BlobWriter.h"
+#include "Storage/ChunkedBufferWriter.h"
+#include "Storage/Clients/BundleStorageClient.h"
+#include "Storage/Clients/KeyValueStorageClient.h"
+
+class FPacketWriter;
+
+/**
+ * Writes blobs into bundles
+ */
+class FBundleWriter final : public FBlobWriter
+{
+public:
+	FBundleWriter(TSharedRef<FKeyValueStorageClient> InStorageClient, FUtf8String InBasePath, const FBundleOptions& InOptions);
+	~FBundleWriter();
+
+	// Inherited from FBlobWriter
+	void Flush() override;
+	TUniquePtr<FBlobWriter> Fork() override;
+	virtual void AddAlias(const FAliasInfo& AliasInfo) override;
+	virtual void AddImport(FBlobHandle Target) override;
+	virtual void AddRef(const FRefName& RefName, const FRefOptions& Options) override;
+	virtual FBlobHandle CompleteBlob(const FBlobType& InType) override;
+	virtual FMutableMemoryView GetOutputBufferAsSpan(size_t UsedSize, size_t DesiredSize) override;
+	virtual void Advance(size_t Size) override;
+
+private:
+	class FPendingExportHandleData;
+	class FPendingPacketHandleData;
+	class FPendingBundleHandleData;
+
+	TSharedRef<FKeyValueStorageClient> StorageClient;
+	FUtf8String BasePath;
+	FBundleOptions Options;
+
+	TBlobHandle<FPendingPacketHandleData> CurrentPacketHandle;
+	TUniquePtr<FPacketWriter> PacketWriter;
+
+	TBlobHandle<FPendingBundleHandleData> CurrentBundleHandle;
+	FChunkedBufferWriter CompressedPacketWriter;
+	TArray<FBlobHandle> CurrentBundleImports;
+
+	TArray<FBlobHandle> BundleReferences;
+	//	std::vector<std::pair<FBlobHandle, FAliasInfo>> PendingExportAliases;
+	//	FChunkedBlobWriter EncodedPacketWriter;
+
+	void StartPacket();
+	void FinishPacket();
+};
\ No newline at end of file
diff --git a/Engine/Source/Developer/Horde/Public/Storage/Bundles/V2/ExportHandle.h b/Engine/Source/Developer/Horde/Public/Storage/Bundles/V2/ExportHandle.h
new file mode 100644
--- /dev/null
+++ b/Engine/Source/Developer/Horde/Public/Storage/Bundles/V2/ExportHandle.h
@@ -0,0 +1,43 @@
+// Copyright Epic Games, Inc. All Rights Reserved.
+
+#pragma once
+
+#include "CoreMinimal.h"
+#include "Storage/BlobHandle.h"
+#include "PacketHandle.h"
+
+class FExportHandleData;
+
+/** Handle to an export within a packet. */
+typedef TBlobHandle<FExportHandleData> FExportHandle;
+
+/** 
+ * Implementation of export handle data.
+ */
+class FExportHandleData final : public FBlobHandleData, public TSharedFromThis<FExportHandleData, ESPMode::ThreadSafe>
+{
+public:
+	static const char Type[];
+
+	FExportHandleData(FPacketHandle InPacket, int32 InExportIdx);
+	FExportHandleData(FPacketHandle InPacket, const FUtf8StringView& InFragment);
+	~FExportHandleData();
+
+	static void AppendIdentifier(FUtf8String& OutBuffer, int32 ExportIdx);
+
+	// Implementation of FBlobHandle
+	virtual bool Equals(const FBlobHandleData& Other) const override;
+	virtual uint32 GetHashCode() const override;
+	virtual FBlobHandle GetOuter() const override;
+	virtual const char* GetType() const override;
+	virtual FBlob Read() const override;
+	virtual bool TryAppendIdentifier(FUtf8String& OutBuffer) const override;
+
+private:
+	static const FUtf8StringView FragmentPrefix;
+
+	FPacketHandle Packet;
+	int32 ExportIdx;
+
+	static bool TryParse(const FUtf8StringView& Fragment, int32& OutExportIdx);
+};
diff --git a/Engine/Source/Developer/Horde/Public/Storage/Bundles/V2/Packet.h b/Engine/Source/Developer/Horde/Public/Storage/Bundles/V2/Packet.h
new file mode 100644
--- /dev/null
+++ b/Engine/Source/Developer/Horde/Public/Storage/Bundles/V2/Packet.h
@@ -0,0 +1,107 @@
+// Copyright Epic Games, Inc. All Rights Reserved.
+
+#pragma once
+
+#include "CoreMinimal.h"
+#include "Containers/Utf8String.h"
+#include "Storage/BlobType.h"
+#include "../BundleCompression.h"
+#include "../../../SharedBufferView.h"
+
+class FChunkedBufferWriter;
+struct FPacketImport;
+struct FPacketExport;
+
+/**
+ * Accessor for data structures stored into a serialized bundle packet.
+ */
+class FPacket
+{
+public:
+	static const FBlobType BlobType;
+
+	FPacket(FSharedBufferView InBuffer);
+	virtual ~FPacket();
+
+	/** Gets the underlying buffer for this packet. */
+	FSharedBufferView GetBuffer() const;
+
+	/** Length of this packet. */
+	size_t GetLength() const;
+
+	/** Gets the number of types in this packet. */
+	size_t GetTypeCount() const;
+
+	/** Gets a type from the packet. */
+	const FBlobType& GetType(size_t TypeIdx) const;
+
+	/** Gets the number of imports in this packet. */
+	size_t GetImportCount() const;
+
+	/** Gets the locator for a particular import. */
+	FPacketImport GetImport(size_t ImportIdx) const;
+
+	/** Gets the number of exports in this packet. */
+	size_t GetExportCount() const;
+
+	/** Gets the bulk data for a particular export. */
+	FPacketExport GetExport(size_t ExportIdx) const;
+
+	/** Encodes a packet. */
+	void Encode(EBundleCompressionFormat Format, FChunkedBufferWriter& Writer);
+
+	/** Decodes a packet from the given data. */
+	static FPacket Decode(const FMemoryView& View);
+
+private:
+	struct FEncodedPacketHeader;
+
+	FSharedBufferView Buffer;
+};
+
+/**
+ * Specifies the path to an imported node
+ */
+struct FPacketImport
+{
+	static constexpr int32 Bias = 3;
+	static constexpr int32 InvalidBaseIdx = -1;
+	static constexpr int32 CurrentPacketBaseIdx = -2;
+	static constexpr int32 CurrentBundleBaseIdx = -3;
+
+	FPacketImport(const FMemoryView& InView);
+	FPacketImport(int32 InBaseIdx, const FUtf8StringView& InFragment);
+
+	/** Gets the index of the base packet import. */
+	int GetBaseIdx() const;
+
+	/** Gets the fragment for this import. */
+	FUtf8StringView GetFragment() const;
+
+private:
+	int32 BaseIdx;
+	FUtf8StringView Fragment;
+};
+
+/*
+ * Data for an exported node in a packet
+ */
+struct FPacketExport
+{
+	FPacketExport(FSharedBufferView InBuffer);
+	~FPacketExport();
+
+	/** Gets the index of this export's type. */
+	size_t GetTypeIndex() const;
+
+	/** Gets the import indexes for this export. */
+	void GetImportIndices(TArray<size_t>& OutIndices) const;
+
+	/** Gets the payload for this export. */
+	FSharedBufferView GetPayload() const;
+
+private:
+	FSharedBufferView Buffer;
+	int32 TypeIdx;
+	FMemoryView Imports;
+};
diff --git a/Engine/Source/Developer/Horde/Public/Storage/Bundles/V2/PacketHandle.h b/Engine/Source/Developer/Horde/Public/Storage/Bundles/V2/PacketHandle.h
new file mode 100644
--- /dev/null
+++ b/Engine/Source/Developer/Horde/Public/Storage/Bundles/V2/PacketHandle.h
@@ -0,0 +1,57 @@
+// Copyright Epic Games, Inc. All Rights Reserved.
+
+#pragma once
+
+#include "CoreMinimal.h"
+#include "Templates/SharedPointer.h"
+#include "Storage/Blob.h"
+#include "Storage/BlobHandle.h"
+#include "Storage/BlobType.h"
+#include "Storage/StorageClient.h"
+
+class FPacketReader;
+class FPacketHandleData;
+
+// Handle to an packet within a bundle. 
+typedef TBlobHandle<FPacketHandleData> FPacketHandle;
+
+// Data for FPacketHandle
+class FPacketHandleData final : public FBlobHandleData, public TSharedFromThis<FPacketHandleData, ESPMode::ThreadSafe>
+{
+public:
+	static const char Type[];
+
+	FPacketHandleData(TSharedRef<FStorageClient> InStorageClient, FBlobHandle InBundle, size_t InPacketOffset, size_t InPacketLength);
+	FPacketHandleData(TSharedRef<FStorageClient> InStorageClient, FBlobHandle InBundle, const FUtf8StringView& InFragment);
+	~FPacketHandleData();
+
+	/** Reads an export from this packet. */
+	FBlob ReadExport(size_t ExportIdx) const;
+
+	/** Reads an export body from this packet. */
+	FSharedBufferView ReadExportBody(size_t ExportIdx) const;
+
+	// Implementation of FBlobHandle
+	virtual bool Equals(const FBlobHandleData& Other) const override;
+	virtual uint32 GetHashCode() const override;
+	virtual FBlobHandle GetOuter() const override;
+	virtual const char* GetType() const override;
+	virtual FBlob Read() const override;
+	virtual bool TryAppendIdentifier(FUtf8String& OutBuffer) const override;
+	virtual FBlobHandle GetFragmentHandle(const FUtf8StringView& Fragment) const override;
+
+private:
+	static const char FragmentPrefix[];
+	static const size_t FragmentPrefixLength;
+
+	TSharedRef<FStorageClient> StorageClient;
+
+	FBlobHandle Bundle;
+	size_t PacketOffset;
+	size_t PacketLength;
+
+	mutable TSharedPtr<FPacketReader> PacketReader;
+
+	const FPacketReader& GetPacketReader() const;
+	static bool TryParse(const FUtf8StringView& Fragment, size_t& OutPacketOffset, size_t& OutPacketLength);
+};
diff --git a/Engine/Source/Developer/Horde/Public/Storage/Bundles/V2/PacketReader.h b/Engine/Source/Developer/Horde/Public/Storage/Bundles/V2/PacketReader.h
new file mode 100644
--- /dev/null
+++ b/Engine/Source/Developer/Horde/Public/Storage/Bundles/V2/PacketReader.h
@@ -0,0 +1,36 @@
+// Copyright Epic Games, Inc. All Rights Reserved.
+
+#pragma once
+
+#include "CoreMinimal.h"
+#include "Storage/StorageClient.h"
+#include "Packet.h"
+#include "PacketHandle.h"
+
+/**
+ * Utility class for constructing BlobData objects from a packet, caching any computed handles to other blobs.
+ */
+class FPacketReader
+{
+public:
+	FPacketReader(TSharedRef<FStorageClient> InStorageClient, FBlobHandle InBundleHandle, FPacketHandle InPacketHandle, FPacket InPacket);
+	~FPacketReader();
+
+	/** Read a blob object for the entire packet. */
+	FBlob ReadPacket() const;
+
+	/** Reads an export from this packet. */
+	FBlob ReadExport(size_t ExportIdx) const;
+
+	/** Reads an export body from this packet. */
+	FSharedBufferView ReadExportBody(size_t ExportIdx) const;
+
+private:
+	TSharedRef<FStorageClient> StorageClient;
+	FBlobHandle BundleHandle;
+	FPacketHandle PacketHandle;
+	FPacket Packet;
+	mutable TArray<FBlobHandle> CachedImportHandles;
+
+	FBlobHandle GetImportHandle(size_t Index) const;
+};
diff --git a/Engine/Source/Developer/Horde/Public/Storage/Bundles/V2/PacketWriter.h b/Engine/Source/Developer/Horde/Public/Storage/Bundles/V2/PacketWriter.h
new file mode 100644
--- /dev/null
+++ b/Engine/Source/Developer/Horde/Public/Storage/Bundles/V2/PacketWriter.h
@@ -0,0 +1,82 @@
+// Copyright Epic Games, Inc. All Rights Reserved.
+
+#pragma once
+
+#include "CoreMinimal.h"
+#include "Storage/BlobHandle.h"
+#include "Packet.h"
+#include "PacketHandle.h"
+#include "Storage/ChunkedBufferWriter.h"
+
+/**
+ * Writes exports into a new bundle packet
+ */
+class FPacketWriter
+{
+public:
+	FPacketWriter(FBlobHandle InBundleHandle, FBlobHandle InPacketHandle);
+	FPacketWriter(const FPacketWriter&) = delete;
+	~FPacketWriter();
+
+	/** Current length of the packet */
+	size_t GetLength() const;
+
+	/** Adds an import to the current blob */
+	void AddImport(FBlobHandle Import);
+
+	/** Gets the number of unique imports current added to this packet */
+	int GetImportCount() const;
+
+	/** Gets a packet import by index */
+	FBlobHandle GetImport(int ImportIdx) const;
+
+	/** Gets data to write new export */
+	FMutableMemoryView GetOutputBuffer(size_t UsedSize, size_t DesiredSize);
+
+	/** Increase the length of the current blob */
+	void Advance(size_t Size);
+
+	/** Writes a new blob to this packet */
+	int CompleteBlob(const FBlobType& Type);
+
+	/** Gets the number of exports currently in this packet */
+	int GetExportCount() const;
+
+	/** Reads data for a blob written to storage */
+	FBlob GetExport(int ExportIdx) const;
+
+	/** Mark the current packet as complete */
+	FPacket CompletePacket();
+
+	FPacketWriter& operator=(const FPacketWriter&) = delete;
+
+private:
+	FBlobHandle BundleHandle;
+	FBlobHandle PacketHandle;
+
+	FSharedBuffer Buffer;
+	uint8* BufferBasePtr;
+	size_t Length;
+
+	size_t NextBlobLength;
+	TArray<int32> NextBlobImports;
+
+	TArray<FBlobType> Types;
+	TArray<uint32> ImportOffsets;
+	FChunkedBufferWriter ImportWriter;
+	TArray<FBlobHandle> ImportHandles;
+	TMap<FBlobHandle, int> ImportMap;
+	TArray<uint32> ExportOffsets;
+
+	FMutableMemoryView GetOutputBufferInternal(size_t UsedSize, size_t DesiredSize);
+	uint8* GetOutputSpanAndAdvance(size_t Length);
+
+	int32 FindOrAddType(FBlobType Type);
+	int32 FindOrAddImport(const FBlobHandle& Handle);
+
+	void WriteTypeTable();
+	void WriteImportTable();
+	void WriteExportTable();
+
+	static size_t Align(size_t Offset);
+};
diff --git a/Engine/Source/Developer/Horde/Public/Storage/ChunkedBufferWriter.h b/Engine/Source/Developer/Horde/Public/Storage/ChunkedBufferWriter.h
new file mode 100644
--- /dev/null
+++ b/Engine/Source/Developer/Horde/Public/Storage/ChunkedBufferWriter.h
@@ -0,0 +1,44 @@
+// Copyright Epic Games, Inc. All Rights Reserved.
+
+#pragma once
+
+#include "CoreMinimal.h"
+#include "BlobWriter.h"
+
+/**
+ * Writes data to a series of memory blocks
+ */
+class FChunkedBufferWriter
+{
+public:
+	FChunkedBufferWriter(size_t InitialSize = 1024);
+	FChunkedBufferWriter(const FChunkedBufferWriter&) = delete;
+	virtual ~FChunkedBufferWriter();
+
+	/** Reset the contents of this writer. */
+	void Reset();
+
+	/** Gets the current written length of this buffer. */
+	size_t GetLength() const;
+
+	/** Get a handle to part of the written buffer */
+	FSharedBufferView Slice(size_t Offset, size_t Length) const;
+
+	/** Gets a view over the underlying buffer. */
+	TArray<FMemoryView> GetView() const;
+
+	/** Copies the entire contents of this writer to another buffer. */
+	void CopyTo(void* Buffer) const;
+
+	/** Gets an output buffer for writing. */
+	FMutableMemoryView GetOutputBuffer(size_t UsedSize, size_t DesiredSize);
+
+	/** Increase the length of the written data. */
+	void Advance(size_t Size);
+
+private:
+	struct FChunk;
+
+	TArray<FChunk> Chunks;
+	size_t WrittenLength;
+};
diff --git a/Engine/Source/Developer/Horde/Public/Storage/Clients/BundleStorageClient.h b/Engine/Source/Developer/Horde/Public/Storage/Clients/BundleStorageClient.h
new file mode 100644
--- /dev/null
+++ b/Engine/Source/Developer/Horde/Public/Storage/Clients/BundleStorageClient.h
@@ -0,0 +1,45 @@
+// Copyright Epic Games, Inc. All Rights Reserved.
+
+#pragma once
+
+#include "CoreMinimal.h"
+#include "Storage/StorageClient.h"
+#include "Storage/BlobHandle.h"
+#include "Storage/Bundles/BundleOptions.h"
+#include "Storage/Clients/KeyValueStorageClient.h"
+
+/**
+ * Base class for an implementation of <see cref="IStorageClient"/>, providing implementations for some common functionality using bundles.
+ */
+class FBundleStorageClient final : public FKeyValueStorageClient
+{
+public:
+	FBundleStorageClient(TSharedRef<FKeyValueStorageClient> InInner);
+	virtual ~FBundleStorageClient() override;
+
+	// FKeyValueStorageClient implementation
+	virtual FBlob ReadBlob(const FBlobLocator& Locator) const override;
+	virtual FBlobHandle WriteBlob(const FUtf8StringView& BasePath, const FBlobType& Type, const TArrayView<FMemoryView>& Data, const TArrayView<FBlobHandle>& Imports) override;
+
+	// FBundleStorageClient implementation
+	virtual FBlobHandle CreateHandle(const FBlobLocator& Locator) const override;
+	virtual TUniquePtr<FBlobWriter> CreateWriter(FUtf8String BasePath) override;
+	TUniquePtr<FBlobWriter> CreateWriter(FUtf8String BasePath, const FBundleOptions& Options);
+
+	virtual void AddAlias(const char* Name, FBlobAlias Alias) override;
+	virtual void RemoveAlias(const char* Name, FBlobHandle Handle) override;
+	virtual void FindAliases(const char* Name, int MaxResults, TArray<FBlobAlias>& OutAliases) override;
+
+	virtual bool DeleteRef(const FRefName& Name) override;
+	virtual FBlobHandle ReadRef(const FRefName& Name, const FRefCacheTime& CacheTime) const override;
+	virtual void WriteRef(const FRefName& Name, const FBlobHandle& Handle, const FRefOptions& Options) override;
+
+private:
+	class FBundleHandleData;
+	friend FBundleHandleData;
+
+	TSharedRef<FKeyValueStorageClient> Inner;
+
+	TBlobHandle<FBundleHandleData> CreateBundleHandle(const FBlobLocator& Locator) const;
+	FBlobHandle GetFragmentHandle(const FBlobHandle& BundleHandle, const FUtf8StringView& Fragment) const;
+};
diff --git a/Engine/Source/Developer/Horde/Public/Storage/Clients/FileStorageClient.h b/Engine/Source/Developer/Horde/Public/Storage/Clients/FileStorageClient.h
new file mode 100644
--- /dev/null
+++ b/Engine/Source/Developer/Horde/Public/Storage/Clients/FileStorageClient.h
@@ -0,0 +1,38 @@
+// Copyright Epic Games, Inc. All Rights Reserved.
+
+#pragma once
+
+#include "CoreMinimal.h"
+#include "KeyValueStorageClient.h"
+#include <filesystem>
+
+/**
+ * Implementation of FStorageClient which writes data to files on disk.
+ */
+class FFileStorageClient : public FKeyValueStorageClient
+{
+public:
+	FFileStorageClient(std::filesystem::path InRootDir);
+	~FFileStorageClient();
+
+	static FBlobLocator ReadRefFromFile(const std::filesystem::path &File);
+	static void WriteRefToFile(const std::filesystem::path& File, const FBlobLocator& Locator);
+
+	std::filesystem::path GetBlobFile(const FBlobLocator& Locator) const;
+	std::filesystem::path GetRefFile(const FRefName& Name) const;
+
+	virtual void AddAlias(const char* Name, FBlobAlias Alias) override;
+	virtual void RemoveAlias(const char* Name, FBlobHandle Handle) override;
+	virtual void FindAliases(const char* Name, int MaxResults, TArray<FBlobAlias>& OutAliases) override;
+
+	virtual bool DeleteRef(const FRefName& Name) override;
+	virtual FBlobHandle ReadRef(const FRefName& Name, const FRefCacheTime& CacheTime) const override;
+	virtual void WriteRef(const FRefName& Name, const FBlobHandle& Handle, const FRefOptions& Options) override;
+
+protected:
+	virtual FBlob ReadBlob(const FBlobLocator& Locator) const override;
+	virtual FBlobHandle WriteBlob(const FUtf8StringView& BasePath, const FBlobType& Type, const TArrayView<FMemoryView>& Data, const TArrayView<FBlobHandle>& Imports) override;
+
+private:
+	std::filesystem::path RootDir;
+};
diff --git a/Engine/Source/Developer/Horde/Public/Storage/Clients/KeyValueStorageClient.h b/Engine/Source/Developer/Horde/Public/Storage/Clients/KeyValueStorageClient.h
new file mode 100644
--- /dev/null
+++ b/Engine/Source/Developer/Horde/Public/Storage/Clients/KeyValueStorageClient.h
@@ -0,0 +1,27 @@
+// Copyright Epic Games, Inc. All Rights Reserved.
+
+#pragma once
+
+#include "CoreMinimal.h"
+#include "Storage/StorageClient.h"
+#include "Storage/BlobHandle.h"
+
+/*
+ * Base class for storage clients that wrap a diirect key/value type store without any merging/splitting.
+ */
+class FKeyValueStorageClient : public FStorageClient
+{
+public:
+	/** Read a single blob from the underlying store. */
+	virtual FBlob ReadBlob(const FBlobLocator& Locator) const = 0;
+
+	/** Write a single blob to the underlying store. */
+	virtual FBlobHandle WriteBlob(const FUtf8StringView& BasePath, const FBlobType& Type, const TArrayView<FMemoryView>& Data, const TArrayView<FBlobHandle>& Imports) = 0;
+
+	// Implementation of FStorageClient
+	virtual FBlobHandle CreateHandle(const FBlobLocator& Locator) const override;
+	virtual TUniquePtr<FBlobWriter> CreateWriter(FUtf8String BasePath = FUtf8String()) override;
+
+private:
+	class FHandleData;
+};
diff --git a/Engine/Source/Developer/Horde/Public/Storage/Clients/MemoryStorageClient.h b/Engine/Source/Developer/Horde/Public/Storage/Clients/MemoryStorageClient.h
new file mode 100644
--- /dev/null
+++ b/Engine/Source/Developer/Horde/Public/Storage/Clients/MemoryStorageClient.h
@@ -0,0 +1,35 @@
+// Copyright Epic Games, Inc. All Rights Reserved.
+
+#pragma once
+
+#include "CoreMinimal.h"
+#include "KeyValueStorageClient.h"
+
+/**
+ * Implementation of FStorageClient which stores data in memory.
+ */
+class FMemoryStorageClient final : public FKeyValueStorageClient
+{
+public:
+	FMemoryStorageClient();
+	~FMemoryStorageClient();
+
+	virtual void AddAlias(const char* Name, FBlobAlias Alias) override;
+	virtual void RemoveAlias(const char* Name, FBlobHandle Handle) override;
+	virtual void FindAliases(const char* Name, int MaxResults, TArray<FBlobAlias>& OutAliases) override;
+
+	virtual bool DeleteRef(const FRefName& Name) override;
+	virtual FBlobHandle ReadRef(const FRefName& Name, const FRefCacheTime& CacheTime) const override;
+	virtual void WriteRef(const FRefName& Name, const FBlobHandle& Target, const FRefOptions& Options) override;
+
+protected:
+	virtual FBlob ReadBlob(const FBlobLocator& Locator) const override;
+	virtual FBlobHandle WriteBlob(const FUtf8StringView& BasePath, const FBlobType& Type, const TArrayView<FMemoryView>& Data, const TArrayView<FBlobHandle>& Imports) override;
+
+private:
+	mutable FCriticalSection CriticalSection;
+
+	TMap<FBlobLocator, FBlob> Blobs;
+	TMap<FRefName, FBlobLocator> Refs;
+	TMap<FUtf8String, TArray<FBlobAlias>> Aliases;
+};
diff --git a/Engine/Source/Developer/Horde/Public/Storage/Nodes/ChunkNode.h b/Engine/Source/Developer/Horde/Public/Storage/Nodes/ChunkNode.h
new file mode 100644
--- /dev/null
+++ b/Engine/Source/Developer/Horde/Public/Storage/Nodes/ChunkNode.h
@@ -0,0 +1,85 @@
+// Copyright Epic Games, Inc. All Rights Reserved.
+
+#pragma once
+
+#include "CoreMinimal.h"
+#include "IO/IoHash.h"
+#include "../BlobType.h"
+#include "../BlobHandle.h"
+#include "../BlobWriter.h"
+#include "../../SharedBufferView.h"
+#include "Hash/Blake3.h"
+
+/**
+ * Base class for chunked data nodes
+ */
+class FChunkNode
+{
+public:
+	virtual ~FChunkNode();
+};
+
+/**
+ * Hashed reference to a chunked data node
+ */
+struct FChunkNodeRef
+{
+	FBlobHandle Handle;
+	FIoHash Hash;
+
+	FChunkNodeRef(FBlobHandle InHandle, const FIoHash& InHash);
+	~FChunkNodeRef();
+};
+
+/**
+ * Chunked data node containing a leaf of chunked data
+ */
+class FLeafChunkNode final : public FChunkNode
+{
+public:
+	static const FBlobType BlobType;
+
+	const FSharedBufferView Buffer;
+
+	FLeafChunkNode(FSharedBufferView InBuffer);
+	virtual ~FLeafChunkNode() override;
+
+	static FLeafChunkNode Read(FBlob Blob);
+	void Write(FBlobWriter& Writer);
+};
+
+/**
+ * An interior file node
+ */
+class FInteriorChunkNode final : public FChunkNode
+{
+public:
+	static const FBlobType BlobType;
+
+	TArray<FChunkNodeRef> Children;
+
+	FInteriorChunkNode();
+	virtual ~FInteriorChunkNode() override;
+
+	static FInteriorChunkNode Read(FBlob Blob);
+	void Write(FBlobWriter& Writer) const;
+};
+
+/**
+ * Utility class for reading data a data stream from a tree of chunk nodes
+ */
+class FChunkNodeReader
+{
+public:
+	FChunkNodeReader(FBlob Blob);
+	FChunkNodeReader(const FBlobHandle& Handle);
+	~FChunkNodeReader();
+
+	bool IsEof() const;
+	FMemoryView GetBuffer() const;
+	void Advance(int32 Length);
+
+private:
+	struct FStackEntry;
+	TArray<FStackEntry> Stack;
+};
diff --git a/Engine/Source/Developer/Horde/Public/Storage/Nodes/DirectoryEntry.h b/Engine/Source/Developer/Horde/Public/Storage/Nodes/DirectoryEntry.h
new file mode 100644
--- /dev/null
+++ b/Engine/Source/Developer/Horde/Public/Storage/Nodes/DirectoryEntry.h
@@ -0,0 +1,33 @@
+// Copyright Epic Games, Inc. All Rights Reserved.
+
+#pragma once
+
+#include "CoreMinimal.h"
+#include "Containers/Utf8String.h"
+#include "../BlobHandle.h"
+#include "IO/IoHash.h"
+
+class FBlobReader;
+class FBlobWriter;
+
+/**
+ * Entry for a directory within a directory node
+ */
+class FDirectoryEntry
+{
+public:
+	FBlobHandle Target;
+	FIoHash TargetHash;
+
+	/** Total size of this directory's contents. */
+	int64 Length;
+
+	/** Name of this directory. */
+	const FUtf8String Name;
+
+	FDirectoryEntry(FBlobHandle InTarget, const FIoHash& InTargetHash, FUtf8String InName, int64 InLength);
+	~FDirectoryEntry();
+
+	static FDirectoryEntry Read(FBlobReader& Reader);
+	void Write(FBlobWriter& Writer) const;
+};
diff --git a/Engine/Source/Developer/Horde/Public/Storage/Nodes/DirectoryNode.h b/Engine/Source/Developer/Horde/Public/Storage/Nodes/DirectoryNode.h
new file mode 100644
--- /dev/null
+++ b/Engine/Source/Developer/Horde/Public/Storage/Nodes/DirectoryNode.h
@@ -0,0 +1,37 @@
+// Copyright Epic Games, Inc. All Rights Reserved.
+
+#pragma once
+
+#include "CoreMinimal.h"
+#include "FileEntry.h"
+#include "DirectoryEntry.h"
+#include "../BlobType.h"
+#include "../StorageClient.h"
+
+/**
+ * Flags for a directory node
+ */
+enum class EDirectoryFlags
+{
+	/** No flags specified. */
+	None = 0,
+};
+
+/**
+ * Stores the contents of a directory in a blob
+ */
+class FDirectoryNode
+{
+public:
+	static const FBlobType BlobType;
+
+	EDirectoryFlags Flags;
+	TMap<FUtf8String, FFileEntry> NameToFile;
+	TMap<FUtf8String, FDirectoryEntry> NameToDirectory;
+
+	FDirectoryNode(EDirectoryFlags InFlags = EDirectoryFlags::None);
+	~FDirectoryNode();
+
+	static FDirectoryNode Read(const FBlob& Blob);
+	FBlobHandle Write(FBlobWriter& Writer) const;
+};
diff --git a/Engine/Source/Developer/Horde/Public/Storage/Nodes/FileEntry.h b/Engine/Source/Developer/Horde/Public/Storage/Nodes/FileEntry.h
new file mode 100644
--- /dev/null
+++ b/Engine/Source/Developer/Horde/Public/Storage/Nodes/FileEntry.h
@@ -0,0 +1,72 @@
+// Copyright Epic Games, Inc. All Rights Reserved.
+
+#pragma once
+
+#include "CoreMinimal.h"
+#include "Containers/Utf8String.h"
+#include "IO/IoHash.h"
+#include "../BlobHandle.h"
+#include "../../SharedBufferView.h"
+
+class FBlobReader;
+class FBlobWriter;
+
+/**
+ * Flags for a file entry
+ */
+enum class EFileEntryFlags
+{
+	/** No other flags set. */
+	None = 0,
+
+	/** Indicates that the referenced file is executable. */
+	Executable = 4,
+
+	/** File should be stored as read-only. */
+	ReadOnly = 8,
+
+	/** File contents are utf-8 encoded text. Client may want to replace line-endings with OS-specific format. */
+	Text = 16,
+
+	/** Used to indicate that custom data is included in the output. Used internally for serialization; not exposed to users. */
+	HasCustomData = 32,
+
+	/** File should be materialized as UTF-16 (but is stored as a UTF-8 source). */
+	Utf16 = 64,
+};
+
+ENUM_CLASS_FLAGS(EFileEntryFlags)
+
+/**
+ * Entry for a file within a directory node
+ */
+class FFileEntry
+{
+public:
+	/** Handle to the root chunked data node. */
+	FBlobHandle Target;
+
+	/** Hash of the root chunked data node. */
+	FIoHash TargetHash;
+
+	/** Name of this file. */
+	const FUtf8String Name;
+
+	/** Flags for this file. */
+	EFileEntryFlags Flags;
+
+	/** Length of this file. */
+	int64 Length;
+
+	/** Hash of the file. */
+	FIoHash Hash;
+
+	/** Custom user data for this file entry. */
+	FSharedBufferView CustomData;
+
+	FFileEntry(FBlobHandle InTarget, const FIoHash& InTargetHash, FUtf8String InName, EFileEntryFlags InFlags, int64 InLength, const FIoHash& InHash, FSharedBufferView InCustomData);
+	~FFileEntry();
+
+	static FFileEntry Read(FBlobReader& Reader);
+	void Write(FBlobWriter& Writer) const;
+};
diff --git a/Engine/Source/Developer/Horde/Public/Storage/RefName.h b/Engine/Source/Developer/Horde/Public/Storage/RefName.h
new file mode 100644
--- /dev/null
+++ b/Engine/Source/Developer/Horde/Public/Storage/RefName.h
@@ -0,0 +1,24 @@
+// Copyright Epic Games, Inc. All Rights Reserved.
+
+#pragma once
+
+#include "CoreMinimal.h"
+#include "Containers/Utf8String.h"
+
+// Identifier for a ref in the storage system. Refs serve as GC roots, and are persistent entry points to expanding data structures within the store.
+struct FRefName
+{
+public:
+	FRefName(FUtf8String Text);
+	~FRefName();
+
+	/** Accessor for the underlying string. */
+	const FUtf8String& GetText() const;
+
+	bool operator==(const FRefName& Other) const;
+	bool operator!=(const FRefName& Other) const;
+	friend uint32 GetTypeHash(const FRefName& RefName);
+
+private:
+	FUtf8String Text;
+};
diff --git a/Engine/Source/Developer/Horde/Public/Storage/RefOptions.h b/Engine/Source/Developer/Horde/Public/Storage/RefOptions.h
new file mode 100644
--- /dev/null
+++ b/Engine/Source/Developer/Horde/Public/Storage/RefOptions.h
@@ -0,0 +1,21 @@
+// Copyright Epic Games, Inc. All Rights Reserved.
+
+#pragma once
+
+#include "CoreMinimal.h"
+
+/**
+ * Options for adding a new ref
+ */ 
+struct FRefOptions
+{
+	static const FRefOptions Default;
+
+	/** Time until a ref is expired. If zero, the ref does not expire. */
+	FTimespan Lifetime;
+
+	/** Whether to extend the remaining lifetime of a ref whenever it is fetched. Defaults to true. */
+	bool bExtend = true;
+
+	FRefOptions();
+};
diff --git a/Engine/Source/Developer/Horde/Public/Storage/StorageClient.h b/Engine/Source/Developer/Horde/Public/Storage/StorageClient.h
new file mode 100644
--- /dev/null
+++ b/Engine/Source/Developer/Horde/Public/Storage/StorageClient.h
@@ -0,0 +1,117 @@
+// Copyright Epic Games, Inc. All Rights Reserved.
+
+#pragma once
+
+#include "CoreMinimal.h"
+#include "Blob.h"
+#include "BlobHandle.h"
+#include "BlobLocator.h"
+#include "RefName.h"
+#include "Templates/SharedPointer.h"
+
+class FBlobAlias;
+struct FAliasInfo;
+class FBlobWriter;
+struct FRefOptions;
+
+/**
+ * Specifies a relative or absolute time value for returning cached ref values
+ */
+struct FRefCacheTime
+{
+	FDateTime Time;
+
+	FRefCacheTime();
+	FRefCacheTime(FTimespan InMaxAge);
+	FRefCacheTime(FDateTime InTime);
+};
+
+/**
+ * Interface for the storage system.
+ */
+class FStorageClient : public TSharedFromThis<FStorageClient, ESPMode::ThreadSafe>
+{
+public:
+	virtual ~FStorageClient();
+
+	//
+	// Blobs
+	//
+
+	/** Creates a new blob handle from a blob locator. */
+	virtual FBlobHandle CreateHandle(const FBlobLocator& Locator) const = 0;
+
+	/** Creates a new writer for storage blobs. */
+	virtual TUniquePtr<FBlobWriter> CreateWriter(FUtf8String BasePath = FUtf8String()) = 0;
+
+	//
+	// Aliases
+	//
+
+	/** Adds an alias to a given blob. */
+	virtual void AddAlias(const char* Name, FBlobAlias Alias) = 0;
+
+	/** Removes an alias from a blob. */
+	virtual void RemoveAlias(const char* Name, FBlobHandle Handle) = 0;
+
+	/** Finds blobs with the given alias. Unlike refs, aliases do not serve as GC roots. */
+	virtual void FindAliases(const char* Name, int MaxResults, TArray<FBlobAlias>& OutAliases) = 0;
+
+	//
+	// Refs
+	//
+
+	/** Deletes a ref. */
+	virtual bool DeleteRef(const FRefName& Name) = 0;
+
+	/** Reads data for a ref from the store. */
+	virtual FBlobHandle ReadRef(const FRefName& Name, const FRefCacheTime& CacheTime) const = 0;
+
+	/** Writes a new ref to the store. */
+	void WriteRef(const FRefName& Name, const FBlobHandle& Handle);
+	virtual void WriteRef(const FRefName& Name, const FBlobHandle& Handle, const FRefOptions& Options) = 0;
+
+protected:
+
+	/** Creates a new locator with a particular base path. */
+	static FBlobLocator CreateLocator(const FUtf8StringView& BasePath);
+};
+
+/**
+ * Target for an aliased blob name
+ */
+class FBlobAlias
+{
+public:
+	/** Handle to the target blob for the alias. */
+	FBlobHandle Target;
+
+	/** Rank for the alias. */
+	int Rank;
+
+	/** Data stored inline with the alias. */
+	FSharedBufferView Data;
+
+	FBlobAlias(FBlobHandle InTarget, int InRank, FSharedBufferView InData);
+	~FBlobAlias();
+};
+
+/** 
+ * Alias for newly added blobs
+ */
+struct FAliasInfo
+{
+	/** Name of the alias. */
+	FUtf8String Name;
+
+	/** Rank for ordering aliases (lower values are preferred). */
+	int Rank;
+
+	/** Inline data to be stored with the alias. */
+	FSharedBufferView Data;
+
+	FAliasInfo(FUtf8String InName, int InRank = 0);
+	FAliasInfo(FUtf8String InName, int InRank, FSharedBufferView InData);
+	~FAliasInfo();
+};
+
