Horde: Add missing copyright notice.

diff --git a/Engine/Source/Developer/Horde/Public/SharedBufferView.h b/Engine/Source/Developer/Horde/Public/SharedBufferView.h
--- a/Engine/Source/Developer/Horde/Public/SharedBufferView.h
+++ b/Engine/Source/Developer/Horde/Public/SharedBufferView.h
@@ -1,3 +1,5 @@
+// Copyright Epic Games, Inc. All Rights Reserved.
+
 #pragma once
 
 #include "Memory/SharedBuffer.h"
