Horde: Merge leaf and interior node type definitions, and implement support for content-defined chunking in C++ storage lib.

diff --git a/Engine/Source/Developer/Horde/Private/Storage/BlobReader.cpp b/Engine/Source/Developer/Horde/Private/Storage/BlobReader.cpp
--- a/Engine/Source/Developer/Horde/Private/Storage/BlobReader.cpp
+++ b/Engine/Source/Developer/Horde/Private/Storage/BlobReader.cpp
@@ -41,6 +41,17 @@ FBlobHandle FBlobReader::ReadImport()
 
 // ------------------------------------------------------------------------
 
+FBlobHandle ReadBlobHandle(FBlobReader& Reader)
+{
+	return Reader.ReadImport();
+}
+
+FBlobHandleWithHash ReadBlobHandleWithHash(FBlobReader& Reader)
+{
+	FBlobHandle Handle = Reader.ReadImport();
+	return FBlobHandleWithHash(MoveTemp(Handle), ReadIoHash(Reader));
+}
+
 int ReadInt32(FBlobReader& Reader)
 {
 	int Value = *(const int*)Reader.GetBuffer();
diff --git a/Engine/Source/Developer/Horde/Private/Storage/BlobWriter.cpp b/Engine/Source/Developer/Horde/Private/Storage/BlobWriter.cpp
--- a/Engine/Source/Developer/Horde/Private/Storage/BlobWriter.cpp
+++ b/Engine/Source/Developer/Horde/Private/Storage/BlobWriter.cpp
@@ -21,6 +21,17 @@ void* FBlobWriter::GetOutputBuffer(size_t Size)
 
 // ------------------------------------------------------------------------------------------
 
+void WriteBlobHandle(FBlobWriter& Writer, FBlobHandle Handle)
+{
+	Writer.AddImport(MoveTemp(Handle));
+}
+
+void WriteBlobHandleWithHash(FBlobWriter& Writer, FBlobHandleWithHash Target)
+{
+	Writer.AddImport(MoveTemp(Target.Handle));
+	WriteIoHash(Writer, Target.Hash);
+}
+
 void WriteIoHash(FBlobWriter& Writer, const FIoHash& Hash)
 {
 	WriteFixedLengthBytes(Writer, &Hash, sizeof(Hash));
diff --git a/Engine/Source/Developer/Horde/Private/Storage/Nodes/ChunkNode.cpp b/Engine/Source/Developer/Horde/Private/Storage/Nodes/ChunkNode.cpp
--- a/Engine/Source/Developer/Horde/Private/Storage/Nodes/ChunkNode.cpp
+++ b/Engine/Source/Developer/Horde/Private/Storage/Nodes/ChunkNode.cpp
@@ -3,80 +3,64 @@
 #include "Storage/Nodes/ChunkNode.h"
 #include "Storage/Blob.h"
 
-FChunkNode::~FChunkNode()
-{
-}
-
-// ----------------------------------------------------------------------
-
-FChunkNodeRef::FChunkNodeRef(FBlobHandle InHandle, const FIoHash& InHash)
-	: Handle(MoveTemp(InHandle))
-	, Hash(InHash)
-{ }
-
-FChunkNodeRef::~FChunkNodeRef()
-{ }
+const FChunkingOptions FChunkingOptions::Default{};
 
 // ----------------------------------------------------------------------
 
-const FBlobType FLeafChunkNode::BlobType(FGuid(0xB27AFB68, 0x4A4B9E20, 0x8A78D8A4, 0x39D49840), 1);
-
-FLeafChunkNode::FLeafChunkNode(FSharedBufferView InBuffer)
-	: Buffer(MoveTemp(InBuffer))
-{ }
+const FBlobType FChunkNode::LeafBlobType(FGuid(0xB27AFB68, 0x4A4B9E20, 0x8A78D8A4, 0x39D49840), 1);
+const FBlobType FChunkNode::InteriorBlobType(FGuid(0xF4DEDDBC, 0x4C7A70CB, 0x11F04783, 0xB9CDCCAF), 2);
 
-FLeafChunkNode::~FLeafChunkNode()
-{ }
-
-FLeafChunkNode FLeafChunkNode::Read(FBlob Blob)
+FChunkNode::FChunkNode()
 {
-	return FLeafChunkNode(MoveTemp(Blob.Data));
 }
 
-void FLeafChunkNode::Write(FBlobWriter& Writer)
+FChunkNode::FChunkNode(TArray<FBlobHandleWithHash> InChildren, FSharedBufferView InData)
+	: Children(MoveTemp(InChildren))
+	, Data(MoveTemp(InData))
 {
-	void* Target = Writer.GetOutputBuffer(Buffer.GetLength());
-	memcpy(Target, Buffer.GetPointer(), Buffer.GetLength());
-	Writer.Advance(Buffer.GetLength());
 }
 
-// ----------------------------------------------------------------------
-
-const FBlobType FInteriorChunkNode::BlobType(FGuid(0xF4DEDDBC, 0x4C7A70CB, 0x11F04783, 0xB9CDCCAF), 2);
-
-FInteriorChunkNode::FInteriorChunkNode()
-{
-}
-
-FInteriorChunkNode::~FInteriorChunkNode()
+FChunkNode::~FChunkNode()
 {
 }
 
-FInteriorChunkNode FInteriorChunkNode::Read(FBlob Blob)
+FChunkNode FChunkNode::Read(FBlob Blob)
 {
-	FInteriorChunkNode Node;
-
-	int32 NumNodes = (int32)(Blob.Data.GetLength() / sizeof(FIoHash));
-	Node.Children.Reserve(NumNodes);
+	FChunkNode Node;
 
-	const FIoHash* Hashes = (const FIoHash*)Blob.Data.GetPointer();
-	for (int32 Idx = 0; Idx < NumNodes; Idx++)
+	FIoHash* Hashes = (FIoHash*)Blob.Data.GetPointer();
+	for (int32 Idx = 0; Idx < Blob.References.Num(); Idx++)
 	{
-		Node.Children.Add(FChunkNodeRef(MoveTemp(Blob.References[Idx]), Hashes[Idx]));
+		Node.Children.Add(FBlobHandleWithHash(MoveTemp(Blob.References[Idx]), Hashes[Idx]));
 	}
 
+	Node.Data = Blob.Data.Slice(sizeof(FIoHash) * Blob.References.Num());
 	return MoveTemp(Node);
 }
 
-void FInteriorChunkNode::Write(FBlobWriter& Writer) const
+FBlobHandleWithHash FChunkNode::Write(FBlobWriter& Writer) const
+{
+	return Write(Writer, Children, Data.GetView());
+}
+
+FBlobHandleWithHash FChunkNode::Write(FBlobWriter& Writer, const TArrayView<const FBlobHandleWithHash>& Children, FMemoryView Data)
 {
-	FIoHash* Hashes = (FIoHash*)Writer.GetOutputBuffer(sizeof(FIoHash) * Children.Num());
+	int32 BufferSize = (sizeof(FIoHash) * Children.Num()) + Data.GetSize();
+	uint8* Buffer = (uint8*)Writer.GetOutputBuffer(BufferSize);
+
+	FIoHash* NextHash = (FIoHash*)Buffer;
 	for (int32 Idx = 0; Idx < Children.Num(); Idx++)
 	{
-		Hashes[Idx] = Children[Idx].Hash;
+		*(NextHash++) = Children[Idx].Hash;
 		Writer.AddImport(Children[Idx].Handle);
 	}
-	Writer.Advance(sizeof(FIoHash) * Children.Num());
+	memcpy(NextHash, Data.GetData(), Data.GetSize());
+
+	FIoHash Hash = FIoHash::HashBuffer(Buffer, BufferSize);
+	Writer.Advance(BufferSize);
+
+	FBlobHandle Handle = Writer.CompleteBlob((Children.Num() == 0)? LeafBlobType : InteriorBlobType);
+	return FBlobHandleWithHash(MoveTemp(Handle), Hash);
 }
 
 // ----------------------------------------------------------------------
@@ -107,7 +91,7 @@ FChunkNodeReader::~FChunkNodeReader()
 {
 }
 
-bool FChunkNodeReader::IsEof() const
+bool FChunkNodeReader::IsComplete() const
 {
 	return Stack.Num() == 0;
 }
@@ -120,8 +104,6 @@ FMemoryView FChunkNodeReader::GetBuffer() const
 	}
 
 	const FStackEntry& StackTop = Stack.Top();
-	check(StackTop.Blob.Type.Guid == FLeafChunkNode::BlobType.Guid);
-
 	FMemoryView View = StackTop.Blob.Data.GetView();
 	return View.Mid(StackTop.Position);
 }
@@ -131,37 +113,126 @@ void FChunkNodeReader::Advance(int32 Size)
 	while (Stack.Num() > 0)
 	{
 		FStackEntry& StackTop = Stack.Top();
-		if (StackTop.Blob.Type.Guid == FLeafChunkNode::BlobType.Guid)
+
+		FMemoryView Data = StackTop.Blob.Data.GetView();
+		if (StackTop.Position == Data.GetSize())
 		{
-			FMemoryView BlobData = StackTop.Blob.Data.GetView();
+			Stack.Pop();
+		}
+		else if (StackTop.Position < sizeof(FIoHash) * StackTop.Blob.References.Num())
+		{
+			FBlobHandle ChildHandle = StackTop.Blob.References[StackTop.Position / sizeof(FIoHash)];
+			StackTop.Position += sizeof(FIoHash);
+			Stack.Add(FStackEntry(ChildHandle->Read()));
+		}
+		else if (Size > 0)
+		{
+			size_t AdvanceSize = FMath::Min<size_t>(Size, Data.GetSize() - StackTop.Position);
+			StackTop.Position += AdvanceSize;
+			Size -= AdvanceSize;
+		}
+		else
+		{
+			break;
+		}
+	}
+}
+
+FChunkNodeReader::operator bool() const
+{
+	return !IsComplete();
+}
 
-			size_t ChunkSize = FMath::Min<size_t>(Size, BlobData.GetSize() - StackTop.Position);
-			StackTop.Position += ChunkSize;
-			Size -= ChunkSize;
+// ----------------------------------------------------------------------
 
-			if (StackTop.Position < BlobData.GetSize())
-			{
-				break;
-			}
+FChunkNodeWriter::FChunkNodeWriter(FBlobWriter& InWriter, const FChunkingOptions& InOptions)
+	: Writer(InWriter)
+	, Options(InOptions)
+	, Threshold((1LL << 32) / Options.TargetChunkSize)
+	, NodeLength(0)
+{
+}
 
-			Stack.Pop();
-		}
-		else if (StackTop.Blob.Type.Guid == FInteriorChunkNode::BlobType.Guid)
+FChunkNodeWriter::~FChunkNodeWriter()
+{
+}
+
+void FChunkNodeWriter::Write(FMemoryView Data)
+{
+	StreamHasher.Update(Data);
+
+	while (Data.GetSize() > 0)
+	{
+		uint8* NodeBuffer = (uint8*)Writer.GetOutputBuffer(Options.MaxChunkSize);
+
+		// Add up to the minimum chunk size
+		if (NodeLength < Options.MinChunkSize)
 		{
-			FBlobHandle ChildHandle = StackTop.Blob.References[StackTop.Position];
+			uint64 AppendLength = FMath::Min<uint64>(Options.MinChunkSize - NodeLength, Data.GetSize());
+			memcpy(NodeBuffer + NodeLength, Data.GetData(), AppendLength);
+			NodeLength += AppendLength;
+			Data = Data.Mid(AppendLength);
+		}
 
-			StackTop.Position++;
-			if (StackTop.Position == StackTop.Blob.References.Num())
+		const uint8* NextSpan = (const uint8*)Data.GetData();
+
+		// Step forwards until reaching the threshold
+		int32 Idx = 0;
+		while (RollingHash.Get() > Threshold)
+		{
+			if (Idx == Data.GetSize())
 			{
-				Stack.Pop();
+				memcpy(NodeBuffer + NodeLength, Data.GetData(), Data.GetSize());
+				NodeLength += Data.GetSize();
+				return;
 			}
 
-			Stack.Add(FStackEntry(ChildHandle->Read()));
-		}
-		else
-		{
-			// Invalid blob type
-			check(false);
+			NodeBuffer[NodeLength] = NextSpan[Idx];
+			RollingHash.Add(NextSpan[Idx]);
+			RollingHash.Sub(NodeBuffer[NodeLength - Options.MinChunkSize]);
+
+			NodeLength++;
+			Idx++;
 		}
+
+		// Write the current node
+		WriteNode();
 	}
 }
+
+FBlobHandleWithHash FChunkNodeWriter::Flush(FIoHash& OutStreamHash)
+{
+	// Finish writing the current node
+	if (NodeLength > 0 || Nodes.Num() == 0)
+	{
+		WriteNode();
+	}
+
+	// Get the hash for the whole stream
+	OutStreamHash = FIoHash(StreamHasher.Finalize());
+
+	// If there was only one node, we don't need an interior node on top of it
+	if (Nodes.Num() == 1)
+	{
+		return Nodes[0];
+	}
+
+	// Create the interior node
+	return FChunkNode::Write(Writer, Nodes, FMemoryView());
+}
+
+void FChunkNodeWriter::WriteNode()
+{
+	// Update the writer with the length of the written data
+	uint8* NodeBuffer = (uint8*)Writer.GetOutputBuffer(NodeLength);
+	FIoHash NodeHash = FIoHash::HashBuffer(NodeBuffer, NodeLength);
+	Writer.Advance(NodeLength);
+
+	// Create the blob handle
+	FBlobHandle NodeHandle = Writer.CompleteBlob(FChunkNode::LeafBlobType);
+	Nodes.Add(FBlobHandleWithHash(MoveTemp(NodeHandle), NodeHash));
+
+	// Clear the current state
+	NodeLength = 0;
+	RollingHash.Reset();
+}
diff --git a/Engine/Source/Developer/Horde/Private/Storage/Nodes/FileEntry.cpp b/Engine/Source/Developer/Horde/Private/Storage/Nodes/FileEntry.cpp
--- a/Engine/Source/Developer/Horde/Private/Storage/Nodes/FileEntry.cpp
+++ b/Engine/Source/Developer/Horde/Private/Storage/Nodes/FileEntry.cpp
@@ -6,9 +6,8 @@
 #include "Storage/BlobWriter.h"
 #include "../../HordePlatform.h"
 
-FFileEntry::FFileEntry(FBlobHandle InTarget, const FIoHash& InTargetHash, FUtf8String InName, EFileEntryFlags InFlags, int64 InLength, const FIoHash& InHash, FSharedBufferView InCustomData)
+FFileEntry::FFileEntry(FBlobHandleWithHash InTarget, FUtf8String InName, EFileEntryFlags InFlags, int64 InLength, const FIoHash& InHash, FSharedBufferView InCustomData)
 	: Target(MoveTemp(InTarget))
-	, TargetHash(InTargetHash)
 	, Name(MoveTemp(InName))
 	, Flags(InFlags)
 	, Length(InLength)
@@ -21,8 +20,7 @@ FFileEntry::~FFileEntry()
 
 FFileEntry FFileEntry::Read(FBlobReader& Reader)
 {
-	FBlobHandle Target = Reader.ReadImport();
-	FIoHash TargetHash = ReadIoHash(Reader);
+	FBlobHandleWithHash Target = ReadBlobHandleWithHash(Reader);
 
 	FUtf8String Name = ReadString(Reader);
 	EFileEntryFlags Flags = (EFileEntryFlags)ReadUnsignedVarInt(Reader);
@@ -37,13 +35,12 @@ FFileEntry FFileEntry::Read(FBlobReader& Reader)
 //		Flags &= ~EFileEntryFlags::HasCustomData;
 	}
 
-	return FFileEntry(MoveTemp(Target), TargetHash, MoveTemp(Name), Flags, Length, Hash, MoveTemp(CustomData));
+	return FFileEntry(MoveTemp(Target), MoveTemp(Name), Flags, Length, Hash, MoveTemp(CustomData));
 }
 
 void FFileEntry::Write(FBlobWriter& Writer) const
 {
-	Writer.AddImport(Target);
-	WriteIoHash(Writer, TargetHash);
+	WriteBlobHandleWithHash(Writer, Target);
 
 	EFileEntryFlags WriteFlags = Flags; // TODO: (CustomData.Length > 0) ? (Flags | EFileEntryFlags::HasCustomData) : (Flags & ~EFileEntryFlags::HasCustomData);
 
diff --git a/Engine/Source/Developer/Horde/Public/Storage/BlobHandle.h b/Engine/Source/Developer/Horde/Public/Storage/BlobHandle.h
--- a/Engine/Source/Developer/Horde/Public/Storage/BlobHandle.h
+++ b/Engine/Source/Developer/Horde/Public/Storage/BlobHandle.h
@@ -5,6 +5,7 @@
 #include "CoreMinimal.h"
 #include "SharedBufferView.h"
 #include "BlobType.h"
+#include "IO/IoHash.h"
 
 struct FBlob;
 class FBlobHandleData;
@@ -110,3 +111,17 @@ public:
 	/** Gets a handle to a nested blob object. */
 	virtual FBlobHandle GetFragmentHandle(const FUtf8StringView& Fragment) const;
 };
+
+/**
+ * Stores a blob handle along with a hash of the target node.
+ */
+struct FBlobHandleWithHash
+{
+	FBlobHandle Handle;
+	FIoHash Hash;
+
+	FBlobHandleWithHash(FBlobHandle InHandle, FIoHash InHash)
+		: Handle(MoveTemp(InHandle))
+		, Hash(InHash)
+	{ }
+};
diff --git a/Engine/Source/Developer/Horde/Public/Storage/BlobReader.h b/Engine/Source/Developer/Horde/Public/Storage/BlobReader.h
--- a/Engine/Source/Developer/Horde/Public/Storage/BlobReader.h
+++ b/Engine/Source/Developer/Horde/Public/Storage/BlobReader.h
@@ -39,6 +39,8 @@ private:
 
 // ------------------------------------------------------------------------
 
+FBlobHandle ReadBlobHandle(FBlobReader& Reader);
+FBlobHandleWithHash ReadBlobHandleWithHash(FBlobReader& Reader);
 int ReadInt32(FBlobReader& Reader);
 FIoHash ReadIoHash(FBlobReader& Reader);
 FMemoryView ReadFixedLengthBytes(FBlobReader& Reader, size_t Length);
diff --git a/Engine/Source/Developer/Horde/Public/Storage/BlobWriter.h b/Engine/Source/Developer/Horde/Public/Storage/BlobWriter.h
--- a/Engine/Source/Developer/Horde/Public/Storage/BlobWriter.h
+++ b/Engine/Source/Developer/Horde/Public/Storage/BlobWriter.h
@@ -51,6 +51,9 @@ public:
 
 // ------------------------------------------------------------------------------------------
 
+void WriteBlobHandle(FBlobWriter& Writer, FBlobHandle Handle);
+void WriteBlobHandleWithHash(FBlobWriter& Writer, FBlobHandleWithHash Target);
+
 void WriteIoHash(FBlobWriter& Writer, const struct FIoHash& Hash);
 
 void WriteFixedLengthBytes(FBlobWriter& Writer, const void* Data, size_t Length);
diff --git a/Engine/Source/Developer/Horde/Public/Storage/Nodes/ChunkNode.h b/Engine/Source/Developer/Horde/Public/Storage/Nodes/ChunkNode.h
--- a/Engine/Source/Developer/Horde/Public/Storage/Nodes/ChunkNode.h
+++ b/Engine/Source/Developer/Horde/Public/Storage/Nodes/ChunkNode.h
@@ -9,60 +9,40 @@
 #include "../BlobWriter.h"
 #include "../../SharedBufferView.h"
 #include "Hash/Blake3.h"
+#include "Hash/BuzHash.h"
 
 /**
- * Base class for chunked data nodes
+ * Options for chunking data
  */
-class FChunkNode
-{
-public:
-	virtual ~FChunkNode();
-};
-
-/**
- * Hashed reference to a chunked data node
- */
-struct FChunkNodeRef
+struct FChunkingOptions
 {
-	FBlobHandle Handle;
-	FIoHash Hash;
+	static const FChunkingOptions Default;
 
-	FChunkNodeRef(FBlobHandle InHandle, const FIoHash& InHash);
-	~FChunkNodeRef();
+	int32 MinChunkSize = 4 * 1024;
+	int32 TargetChunkSize = 64 * 1024;
+	int32 MaxChunkSize = 128 * 1024;
 };
 
 /**
- * Chunked data node containing a leaf of chunked data
+ * Node containing a chunk of data
  */
-class FLeafChunkNode final : public FChunkNode
+class FChunkNode
 {
 public:
-	static const FBlobType BlobType;
+	static const FBlobType LeafBlobType;
+	static const FBlobType InteriorBlobType;
 
-	const FSharedBufferView Buffer;
-
-	FLeafChunkNode(FSharedBufferView InBuffer);
-	virtual ~FLeafChunkNode() override;
-
-	static FLeafChunkNode Read(FBlob Blob);
-	void Write(FBlobWriter& Writer);
-};
-
-/**
- * An interior file node
- */
-class FInteriorChunkNode final : public FChunkNode
-{
-public:
-	static const FBlobType BlobType;
+	TArray<FBlobHandleWithHash> Children;
+	FSharedBufferView Data;
 
-	TArray<FChunkNodeRef> Children;
+	FChunkNode();
+	FChunkNode(TArray<FBlobHandleWithHash> InChildren, FSharedBufferView InData);
+	~FChunkNode();
 
-	FInteriorChunkNode();
-	virtual ~FInteriorChunkNode() override;
+	static FChunkNode Read(FBlob Blob);
+	FBlobHandleWithHash Write(FBlobWriter& Writer) const;
 
-	static FInteriorChunkNode Read(FBlob Blob);
-	void Write(FBlobWriter& Writer) const;
+	static FBlobHandleWithHash Write(FBlobWriter& Writer, const TArrayView<const FBlobHandleWithHash>& Children, FMemoryView Data);
 };
 
 /**
@@ -75,11 +55,41 @@ public:
 	FChunkNodeReader(const FBlobHandle& Handle);
 	~FChunkNodeReader();
 
-	bool IsEof() const;
+	bool IsComplete() const;
 	FMemoryView GetBuffer() const;
 	void Advance(int32 Length);
 
+	operator bool() const;
+
 private:
 	struct FStackEntry;
 	TArray<FStackEntry> Stack;
 };
+
+/**
+ * Utility class for writing new data to a tree of chunk nodes
+ */
+class FChunkNodeWriter
+{
+public:
+	FChunkNodeWriter(FBlobWriter& InWriter, const FChunkingOptions& InOptions = FChunkingOptions::Default);
+	~FChunkNodeWriter();
+
+	void Write(FMemoryView Data);
+
+	FBlobHandleWithHash Flush(FIoHash& OutStreamHash);
+
+private:
+	FBlobWriter& Writer;
+	const FChunkingOptions& Options;
+
+	FBuzHash RollingHash;
+	uint32 Threshold;
+	int32 NodeLength;
+
+	FBlake3 StreamHasher;
+
+	TArray<FBlobHandleWithHash> Nodes;
+
+	void WriteNode();
+};
diff --git a/Engine/Source/Developer/Horde/Public/Storage/Nodes/FileEntry.h b/Engine/Source/Developer/Horde/Public/Storage/Nodes/FileEntry.h
--- a/Engine/Source/Developer/Horde/Public/Storage/Nodes/FileEntry.h
+++ b/Engine/Source/Developer/Horde/Public/Storage/Nodes/FileEntry.h
@@ -44,10 +44,7 @@ class FFileEntry
 {
 public:
 	/** Handle to the root chunked data node. */
-	FBlobHandle Target;
-
-	/** Hash of the root chunked data node. */
-	FIoHash TargetHash;
+	FBlobHandleWithHash Target;
 
 	/** Name of this file. */
 	const FUtf8String Name;
@@ -64,7 +61,7 @@ public:
 	/** Custom user data for this file entry. */
 	FSharedBufferView CustomData;
 
-	FFileEntry(FBlobHandle InTarget, const FIoHash& InTargetHash, FUtf8String InName, EFileEntryFlags InFlags, int64 InLength, const FIoHash& InHash, FSharedBufferView InCustomData);
+	FFileEntry(FBlobHandleWithHash InTarget, FUtf8String InName, EFileEntryFlags InFlags, int64 InLength, const FIoHash& InHash, FSharedBufferView InCustomData);
 	~FFileEntry();
 
 	static FFileEntry Read(FBlobReader& Reader);
diff --git a/Engine/Source/Runtime/Core/Private/Hash/BuzHash.cpp b/Engine/Source/Runtime/Core/Private/Hash/BuzHash.cpp
new file mode 100644
--- /dev/null
+++ b/Engine/Source/Runtime/Core/Private/Hash/BuzHash.cpp
@@ -0,0 +1,31 @@
+// Copyright Epic Games, Inc. All Rights Reserved.
+
+#include "Hash/BuzHash.h"
+
+const uint32 FBuzHash::Table[256] = 
+{
+	0x458be752, 0xc10748cc, 0xfbbcdbb8, 0x6ded5b68, 0xb10a82b5, 0x20d75648, 0xdfc5665f, 0xa8428801, 0x7ebf5191, 0x841135c7, 0x65cc53b3,
+	0x280a597c, 0x16f60255, 0xc78cbc3e, 0x294415f5, 0xb938d494, 0xec85c4e6, 0xb7d33edc, 0xe549b544, 0xfdeda5aa, 0x882bf287, 0x3116737c,
+	0x05569956, 0xe8cc1f68, 0x0806ac5e, 0x22a14443, 0x15297e10, 0x50d090e7, 0x4ba60f6f, 0xefd9f1a7, 0x5c5c885c, 0x82482f93, 0x9bfd7c64,
+	0x0b3e7276, 0xf2688e77, 0x8fad8abc, 0xb0509568, 0xf1ada29f, 0xa53efdfe, 0xcb2b1d00, 0xf2a9e986, 0x6463432b, 0x95094051, 0x5a223ad2,
+	0x9be8401b, 0x61e579cb, 0x1a556a14, 0x5840fdc2, 0x9261ddf6, 0xcde002bb, 0x52432bb0, 0xbf17373e, 0x7b7c222f, 0x2955ed16, 0x9f10ca59,
+	0xe840c4c9, 0xccabd806, 0x14543f34, 0x1462417a, 0x0d4a1f9c, 0x087ed925, 0xd7f8f24c, 0x7338c425, 0xcf86c8f5, 0xb19165cd, 0x9891c393,
+	0x325384ac, 0x0308459d, 0x86141d7e, 0xc922116a, 0xe2ffa6b6, 0x53f52aed, 0x2cd86197, 0xf5b9f498, 0xbf319c8f, 0xe0411fae, 0x977eb18c,
+	0xd8770976, 0x9833466a, 0xc674df7f, 0x8c297d45, 0x8ca48d26, 0xc49ed8e2, 0x7344f874, 0x556f79c7, 0x6b25eaed, 0xa03e2b42, 0xf68f66a4,
+	0x8e8b09a2, 0xf2e0e62a, 0x0d3a9806, 0x9729e493, 0x8c72b0fc, 0x160b94f6, 0x450e4d3d, 0x7a320e85, 0xbef8f0e1, 0x21d73653, 0x4e3d977a,
+	0x1e7b3929, 0x1cc6c719, 0xbe478d53, 0x8d752809, 0xe6d8c2c6, 0x275f0892, 0xc8acc273, 0x4cc21580, 0xecc4a617, 0xf5f7be70, 0xe795248a,
+	0x375a2fe9, 0x425570b6, 0x8898dcf8, 0xdc2d97c4, 0x0106114b, 0x364dc22f, 0x1e0cad1f, 0xbe63803c, 0x5f69fac2, 0x4d5afa6f, 0x1bc0dfb5,
+	0xfb273589, 0x0ea47f7b, 0x3c1c2b50, 0x21b2a932, 0x6b1223fd, 0x2fe706a8, 0xf9bd6ce2, 0xa268e64e, 0xe987f486, 0x3eacf563, 0x1ca2018c,
+	0x65e18228, 0x2207360a, 0x57cf1715, 0x34c37d2b, 0x1f8f3cde, 0x93b657cf, 0x31a019fd, 0xe69eb729, 0x8bca7b9b, 0x4c9d5bed, 0x277ebeaf,
+	0xe0d8f8ae, 0xd150821c, 0x31381871, 0xafc3f1b0, 0x927db328, 0xe95effac, 0x305a47bd, 0x426ba35b, 0x1233af3f, 0x686a5b83, 0x50e072e5,
+	0xd9d3bb2a, 0x8befc475, 0x487f0de6, 0xc88dff89, 0xbd664d5e, 0x971b5d18, 0x63b14847, 0xd7d3c1ce, 0x7f583cf3, 0x72cbcb09, 0xc0d0a81c,
+	0x7fa3429b, 0xe9158a1b, 0x225ea19a, 0xd8ca9ea3, 0xc763b282, 0xbb0c6341, 0x020b8293, 0xd4cd299d, 0x58cfa7f8, 0x91b4ee53, 0x37e4d140,
+	0x95ec764c, 0x30f76b06, 0x5ee68d24, 0x679c8661, 0xa41979c2, 0xf2b61284, 0x4fac1475, 0x0adb49f9, 0x19727a23, 0x15a7e374, 0xc43a18d5,
+	0x3fb1aa73, 0x342fc615, 0x924c0793, 0xbee2d7f0, 0x8a279de9, 0x4aa2d70c, 0xe24dd37f, 0xbe862c0b, 0x177c22c2, 0x5388e5ee, 0xcd8a7510,
+	0xf901b4fd, 0xdbc13dbc, 0x6c0bae5b, 0x64efe8c7, 0x48b02079, 0x80331a49, 0xca3d8ae6, 0xf3546190, 0xfed7108b, 0xc49b941b, 0x32baf4a9,
+	0xeb833a4a, 0x88a3f1a5, 0x3a91ce0a, 0x3cc27da1, 0x7112e684, 0x4a3096b1, 0x3794574c, 0xa3c8b6f3, 0x1d213941, 0x6e0a2e00, 0x233479f1,
+	0x0f4cd82f, 0x6093edd2, 0x5d7d209e, 0x464fe319, 0xd4dcac9e, 0x0db845cb, 0xfb5e4bc3, 0xe0256ce1, 0x09fb4ed1, 0x0914be1e, 0xa5bdb2c3,
+	0xc6eb57bb, 0x30320350, 0x3f397e91, 0xa67791bc, 0x86bc0e2c, 0xefa0a7e2, 0xe9ff7543, 0xe733612c, 0xd185897b, 0x329e5388, 0x91dd236b,
+	0x2ecb0d93, 0xf4d82a3d, 0x35b5c03f, 0xe4e606f0, 0x05b21843, 0x37b45964, 0x5eff22f4, 0x6027f4cc, 0x77178b3c, 0xae507131, 0x7bf7cabc,
+	0xf9c18d66, 0x593ade65, 0xd95ddf11,
+};
diff --git a/Engine/Source/Runtime/Core/Public/Hash/BuzHash.h b/Engine/Source/Runtime/Core/Public/Hash/BuzHash.h
new file mode 100644
--- /dev/null
+++ b/Engine/Source/Runtime/Core/Public/Hash/BuzHash.h
@@ -0,0 +1,64 @@
+// Copyright Epic Games, Inc. All Rights Reserved.
+
+#pragma once
+
+#include "CoreTypes.h"
+#include "Memory/MemoryView.h"
+
+/**
+ * Rolling hash function. Allows efficient computation of the hash of a window into a block of data, which can be moved
+ * through the data byte by byte.
+ */
+struct FBuzHash
+{
+public:
+	/** Reset the state of the hasher. */
+	FORCEINLINE void Reset() { *this = FBuzHash(); }
+
+	/** Gets the current hash value. */
+	FORCEINLINE uint32 Get() const { return State; }
+
+	/** Gets the size of the window. */
+	FORCEINLINE uint64 GetWindowSize() const { return Count; }
+
+	/** Appends a byte to the start of the window and updates the hash. */
+	FORCEINLINE void Add(uint8 X)
+	{
+		State = Rol32(State, 1) ^ Table[X];
+		Count++;
+	}
+
+	/** Adds an array of bytes to the start of the window. */
+	FORCEINLINE void Add(FMemoryView View)
+	{
+		Add((const uint8*)View.GetData(), View.GetSize());
+	}
+
+	/** Adds an array of bytes to the start of the window. */
+	FORCEINLINE void Add(const uint8* Data, uint64 Size)
+	{
+		for (uint64 Idx = 0; Idx < Size; ++Idx)
+		{
+			Add(Data[Idx]);
+		}
+	}
+
+	/** Removes a byte from the back of the sliding window and updates the hash. */
+	FORCEINLINE void Sub(uint8 X)
+	{
+		State = State ^ Rol32(Table[X], uint32(Count - 1));
+		Count--;
+	}
+
+private:
+	CORE_API static const uint32 Table[256];
+
+	uint64 Count = 0;
+	uint32 State = 0;
+
+	static FORCEINLINE uint32 Rol32(uint32 V, uint32 N)
+	{
+		N &= 31;
+		return ((V) << (N)) | ((V) >> (32 - N));
+	}
+};
\ No newline at end of file
