Horde: Add a typed HTTP client for communicating with the Horde server. Removes boilerplate code for getting an access token, configuring request objects, and parsing responses. Synchronous for now.

diff --git a/Engine/Source/Developer/DesktopPlatform/Private/DesktopPlatformBase.cpp b/Engine/Source/Developer/DesktopPlatform/Private/DesktopPlatformBase.cpp
--- a/Engine/Source/Developer/DesktopPlatform/Private/DesktopPlatformBase.cpp
+++ b/Engine/Source/Developer/DesktopPlatform/Private/DesktopPlatformBase.cpp
@@ -820,7 +820,10 @@ bool FDesktopPlatformBase::GetOidcAccessToken(const FString& RootDir, const FStr
 	FString Arguments = TEXT(" ");
 	Arguments += FString::Printf(TEXT(" --Service=\"%s\""), *ProviderIdentifier);
 	Arguments += FString::Printf(TEXT(" --OutFile=\"%s\""), *ResultFilePath);
-	Arguments += FString::Printf(TEXT(" --project=\"%s\""),  *IFileManager::Get().ConvertToAbsolutePathForExternalAppForRead(*FPaths::ProjectDir()));
+	if (ProjectFileName.Len() > 0)
+	{
+		Arguments += FString::Printf(TEXT(" --project=\"%s\""), *IFileManager::Get().ConvertToAbsolutePathForExternalAppForRead(*FPaths::GetPath(*ProjectFileName)));
+	}
 	FString UnattendedArguments = Arguments;
 	UnattendedArguments += TEXT(" --Unattended=true");
 
diff --git a/Engine/Source/Developer/Horde/Horde.Build.cs b/Engine/Source/Developer/Horde/Horde.Build.cs
--- a/Engine/Source/Developer/Horde/Horde.Build.cs
+++ b/Engine/Source/Developer/Horde/Horde.Build.cs
@@ -7,6 +7,6 @@ public class Horde : ModuleRules
 {
 	public Horde(ReadOnlyTargetRules Target) : base(Target)
 	{
-		PrivateDependencyModuleNames.AddRange(new string[] { "Core" });
+		PrivateDependencyModuleNames.AddRange(new string[] { "Core", "HTTP", "Json", "DesktopPlatform" });
 	}
 }
diff --git a/Engine/Source/Developer/Horde/Private/HordeHttpClient.cpp b/Engine/Source/Developer/Horde/Private/HordeHttpClient.cpp
new file mode 100644
--- /dev/null
+++ b/Engine/Source/Developer/Horde/Private/HordeHttpClient.cpp
@@ -0,0 +1,70 @@
+// Copyright Epic Games, Inc. All Rights Reserved.
+
+#include "HordeHttpClient.h"
+#include "Server/ServerMessages.h"
+#include "Interfaces/IHttpResponse.h"
+#include "DesktopPlatformModule.h"
+#include "IDesktopPlatform.h"
+#include "HttpModule.h"
+#include "Misc/Paths.h"
+
+FHordeHttpClient::FHordeHttpClient(FString InServerUrl)
+	: ServerUrl(InServerUrl)
+{
+}
+
+FHordeHttpClient::~FHordeHttpClient()
+{
+}
+
+bool FHordeHttpClient::LoginWithOidc(const TCHAR* Profile, bool bUnattended, FFeedbackContext* Warn)
+{
+	IDesktopPlatform* DesktopPlatform = FDesktopPlatformModule::Get();
+
+	FString NewToken;
+	FDateTime ExpiresAt;
+	bool bWasInteractive = false;
+
+	if (!DesktopPlatform->GetOidcAccessToken(FPaths::RootDir(), FString(), Profile, bUnattended, Warn, NewToken, ExpiresAt, bWasInteractive))
+	{
+		return false;
+	}
+
+	Token = NewToken;
+	return true;
+}
+
+TSharedRef<IHttpRequest> FHordeHttpClient::CreateRequest(const TCHAR* Verb, const TCHAR* Path)
+{
+	FHttpModule& Module = FHttpModule::Get();
+
+	TSharedRef<IHttpRequest> Request = Module.CreateRequest();
+	Request->SetVerb(Verb);
+	Request->SetURL(ServerUrl / Path);
+	if (Token.Len() > 0)
+	{
+		Request->SetHeader(TEXT("Authorization"), FString::Printf(TEXT("Bearer %s"), *Token));
+	}
+
+	return Request;
+}
+
+TSharedRef<IHttpResponse> FHordeHttpClient::Get(const TCHAR* Path)
+{
+	return ExecuteRequest(CreateRequest(TEXT("GET"), Path));
+}
+
+TSharedRef<IHttpResponse> FHordeHttpClient::ExecuteRequest(TSharedRef<IHttpRequest> Request)
+{
+	verify(Request->ProcessRequest());
+
+	for (;;)
+	{
+		FHttpResponsePtr Response = Request->GetResponse();
+		if (Response)
+		{
+			return Response.ToSharedRef();
+		}
+		FPlatformProcess::Sleep(0.05f);
+	}
+}
diff --git a/Engine/Source/Developer/Horde/Private/Server/ServerMessages.cpp b/Engine/Source/Developer/Horde/Private/Server/ServerMessages.cpp
new file mode 100644
--- /dev/null
+++ b/Engine/Source/Developer/Horde/Private/Server/ServerMessages.cpp
@@ -0,0 +1,37 @@
+// Copyright Epic Games, Inc. All Rights Reserved.
+
+#include "Server/ServerMessages.h"
+
+FGetServerInfoResponse::FGetServerInfoResponse()
+{
+}
+
+FGetServerInfoResponse::~FGetServerInfoResponse()
+{
+}
+
+void FGetServerInfoResponse::Serialize(FJsonSerializerBase& Serializer, bool bFlatObject)
+{
+	JSON_SERIALIZE("ServerVersion", ServerVersion);
+	JSON_SERIALIZE("AgentVersion", AgentVersion);
+	JSON_SERIALIZE("OsDescription", OsDescription);
+}
+
+// -------------------------------------
+
+FGetAuthConfigResponse::FGetAuthConfigResponse()
+{
+}
+
+FGetAuthConfigResponse::~FGetAuthConfigResponse()
+{
+}
+
+void FGetAuthConfigResponse::Serialize(FJsonSerializerBase& Serializer, bool bFlatObject)
+{
+	JSON_SERIALIZE("Method", Method);
+	JSON_SERIALIZE("ServerUrl", ServerUrl);
+	JSON_SERIALIZE("ClientId", ClientId);
+	JSON_SERIALIZE_ARRAY("LocalRedirectUrls", LocalRedirectUrls);
+}
+
diff --git a/Engine/Source/Developer/Horde/Public/HordeHttpClient.h b/Engine/Source/Developer/Horde/Public/HordeHttpClient.h
new file mode 100644
--- /dev/null
+++ b/Engine/Source/Developer/Horde/Public/HordeHttpClient.h
@@ -0,0 +1,34 @@
+// Copyright Epic Games, Inc. All Rights Reserved.
+
+#pragma once
+
+#include "CoreMinimal.h"
+#include "Containers/UnrealString.h"
+#include "Interfaces/IHttpRequest.h"
+#include "Interfaces/IHttpResponse.h"
+
+class HORDE_API FHordeHttpClient
+{
+public:
+	FHordeHttpClient(FString InServerUrl);
+	~FHordeHttpClient();
+
+	bool LoginWithOidc(const TCHAR* Profile, bool bUnattended, FFeedbackContext* Warn = nullptr);
+	
+	TSharedRef<IHttpRequest> CreateRequest(const TCHAR* Verb, const TCHAR* Path);
+	TSharedRef<IHttpResponse> Get(const TCHAR* Path);
+
+	template<typename T>
+	T Get(const TCHAR* Path)
+	{
+		T Result;
+		Result.FromJson(Get(Path)->GetContentAsString());
+		return Result;
+	}
+
+	static TSharedRef<IHttpResponse> ExecuteRequest(TSharedRef<IHttpRequest> Request);
+
+private:
+	FString ServerUrl;
+	FString Token;
+};
diff --git a/Engine/Source/Developer/Horde/Public/Server/ServerMessages.h b/Engine/Source/Developer/Horde/Public/Server/ServerMessages.h
new file mode 100644
--- /dev/null
+++ b/Engine/Source/Developer/Horde/Public/Server/ServerMessages.h
@@ -0,0 +1,44 @@
+// Copyright Epic Games, Inc. All Rights Reserved.
+
+#pragma once
+
+#include "CoreMinimal.h"
+#include "Serialization/JsonSerializerMacros.h"
+
+/** Server Info */
+struct FGetServerInfoResponse : FJsonSerializable
+{
+	/** Server version info */
+	FString ServerVersion;
+
+	/** The current agent version string */
+	FString AgentVersion;
+
+	/** The operating system server is hosted on */
+	FString OsDescription;
+
+	FGetServerInfoResponse();
+	virtual ~FGetServerInfoResponse() override;
+	virtual void Serialize(FJsonSerializerBase& Serializer, bool bFlatObject) override;
+};
+
+/** Describes the auth config for this server */
+struct FGetAuthConfigResponse : FJsonSerializable
+{
+	/** Issuer for tokens from the auth provider */
+	FString Method;
+
+	/** Issuer for tokens from the auth provider */
+	FString ServerUrl;
+
+	/** Client id for the OIDC authority */
+	FString ClientId;
+
+	/** Optional redirect url provided to OIDC login for external tools (typically to a local server) */
+	TArray<FString> LocalRedirectUrls;
+
+	FGetAuthConfigResponse();
+	virtual ~FGetAuthConfigResponse() override;
+	virtual void Serialize(FJsonSerializerBase& Serializer, bool bFlatObject) override;
+};
+
