Export all public Horde interfaces:
- Add HORDE_API specifier to all public interfaces, structures, and functions.
- Add deleted-assignment-operator to FChunkedBufferWriter.

#rnx

diff --git a/Engine/Source/Developer/Horde/Private/Storage/BlobReader.cpp b/Engine/Source/Developer/Horde/Private/Storage/BlobReader.cpp
--- a/Engine/Source/Developer/Horde/Private/Storage/BlobReader.cpp
+++ b/Engine/Source/Developer/Horde/Private/Storage/BlobReader.cpp
@@ -41,25 +41,25 @@ FBlobHandle FBlobReader::ReadImport()
 
 // ------------------------------------------------------------------------
 
-FBlobHandle ReadBlobHandle(FBlobReader& Reader)
+HORDE_API FBlobHandle ReadBlobHandle(FBlobReader& Reader)
 {
 	return Reader.ReadImport();
 }
 
-FBlobHandleWithHash ReadBlobHandleWithHash(FBlobReader& Reader)
+HORDE_API FBlobHandleWithHash ReadBlobHandleWithHash(FBlobReader& Reader)
 {
 	FBlobHandle Handle = Reader.ReadImport();
 	return FBlobHandleWithHash(MoveTemp(Handle), ReadIoHash(Reader));
 }
 
-int ReadInt32(FBlobReader& Reader)
+HORDE_API int ReadInt32(FBlobReader& Reader)
 {
 	int Value = *(const int*)Reader.GetBuffer();
 	Reader.Advance(sizeof(int));
 	return Value;
 }
 
-FIoHash ReadIoHash(FBlobReader& Reader)
+HORDE_API FIoHash ReadIoHash(FBlobReader& Reader)
 {
 	FIoHash Hash;
 	memcpy(&Hash, Reader.GetBuffer(), sizeof(FIoHash));
@@ -67,14 +67,14 @@ FIoHash ReadIoHash(FBlobReader& Reader)
 	return Hash;
 }
 
-FMemoryView ReadFixedLengthBytes(FBlobReader& Reader, size_t Length)
+HORDE_API FMemoryView ReadFixedLengthBytes(FBlobReader& Reader, size_t Length)
 {
 	FMemoryView View = Reader.GetView();
 	Reader.Advance(Length);
 	return View.Left(Length);
 }
 
-size_t ReadUnsignedVarInt(FBlobReader& Reader)
+HORDE_API size_t ReadUnsignedVarInt(FBlobReader& Reader)
 {
 	// Figure out the length of the buffer
 	const unsigned char* Data = Reader.GetBuffer();
@@ -92,13 +92,13 @@ size_t ReadUnsignedVarInt(FBlobReader& Reader)
 	return value;
 }
 
-FUtf8String ReadString(FBlobReader& Reader)
+HORDE_API FUtf8String ReadString(FBlobReader& Reader)
 {
 	FMemoryView String = ReadStringSpan(Reader);
 	return FUtf8String::ConstructFromPtrSize((const UTF8CHAR*)String.GetData(), String.GetSize());
 }
 
-FMemoryView ReadStringSpan(FBlobReader& Reader)
+HORDE_API FMemoryView ReadStringSpan(FBlobReader& Reader)
 {
 	size_t Length = ReadUnsignedVarInt(Reader);
 	return ReadFixedLengthBytes(Reader, Length);
diff --git a/Engine/Source/Developer/Horde/Private/Storage/BlobWriter.cpp b/Engine/Source/Developer/Horde/Private/Storage/BlobWriter.cpp
--- a/Engine/Source/Developer/Horde/Private/Storage/BlobWriter.cpp
+++ b/Engine/Source/Developer/Horde/Private/Storage/BlobWriter.cpp
@@ -21,35 +21,35 @@ void* FBlobWriter::GetOutputBuffer(size_t Size)
 
 // ------------------------------------------------------------------------------------------
 
-void WriteBlobHandle(FBlobWriter& Writer, FBlobHandle Handle)
+HORDE_API void WriteBlobHandle(FBlobWriter& Writer, FBlobHandle Handle)
 {
 	Writer.AddImport(MoveTemp(Handle));
 }
 
-void WriteBlobHandleWithHash(FBlobWriter& Writer, FBlobHandleWithHash Target)
+HORDE_API void WriteBlobHandleWithHash(FBlobWriter& Writer, FBlobHandleWithHash Target)
 {
 	Writer.AddImport(MoveTemp(Target.Handle));
 	WriteIoHash(Writer, Target.Hash);
 }
 
-void WriteIoHash(FBlobWriter& Writer, const FIoHash& Hash)
+HORDE_API void WriteIoHash(FBlobWriter& Writer, const FIoHash& Hash)
 {
 	WriteFixedLengthBytes(Writer, &Hash, sizeof(Hash));
 }
 
-void WriteFixedLengthBytes(FBlobWriter& Writer, const void* Data, size_t Length)
+HORDE_API void WriteFixedLengthBytes(FBlobWriter& Writer, const void* Data, size_t Length)
 {
 	void* Target = Writer.GetOutputBuffer(Length);
 	memcpy(Target, Data, Length);
 	Writer.Advance(Length);
 }
 
-void WriteFixedLengthBytes(FBlobWriter& Writer, const FMemoryView& View)
+HORDE_API void WriteFixedLengthBytes(FBlobWriter& Writer, const FMemoryView& View)
 {
 	WriteFixedLengthBytes(Writer, View.GetData(), View.GetSize());
 }
 
-size_t MeasureUnsignedVarInt(size_t Value)
+HORDE_API size_t MeasureUnsignedVarInt(size_t Value)
 {
 	check(Value == (unsigned int)Value);
 
@@ -63,21 +63,21 @@ size_t MeasureUnsignedVarInt(size_t Value)
 	}
 }
 
-size_t WriteUnsignedVarInt(void* Buffer, size_t Value)
+HORDE_API size_t WriteUnsignedVarInt(void* Buffer, size_t Value)
 {
 	size_t ByteCount = MeasureUnsignedVarInt(Value);
 	WriteUnsignedVarIntWithKnownLength(Buffer, Value, ByteCount);
 	return ByteCount;
 }
 
-void WriteUnsignedVarInt(FBlobWriter& Writer, size_t Value)
+HORDE_API void WriteUnsignedVarInt(FBlobWriter& Writer, size_t Value)
 {
 	size_t ByteCount = MeasureUnsignedVarInt(Value);
 	WriteUnsignedVarIntWithKnownLength(Writer.GetOutputBuffer(ByteCount), Value, ByteCount);
 	Writer.Advance(ByteCount);
 }
 
-void WriteUnsignedVarIntWithKnownLength(void* Buffer, size_t Value, size_t NumBytes)
+HORDE_API void WriteUnsignedVarIntWithKnownLength(void* Buffer, size_t Value, size_t NumBytes)
 {
 	uint8* ByteBuffer = (uint8*)Buffer;
 	for (size_t Idx = 1; Idx < NumBytes; Idx++)
@@ -88,28 +88,28 @@ void WriteUnsignedVarIntWithKnownLength(void* Buffer, size_t Value, size_t NumBy
 	ByteBuffer[0] = (unsigned char)((0xff << (9 - (int)NumBytes)) | (unsigned char)Value);
 }
 
-size_t MeasureString(const char* Text)
+HORDE_API size_t MeasureString(const char* Text)
 {
 	return MeasureString(FUtf8StringView((const UTF8CHAR*)Text));
 }
 
-size_t MeasureString(const FUtf8StringView& Text)
+HORDE_API size_t MeasureString(const FUtf8StringView& Text)
 {
 	return MeasureUnsignedVarInt(Text.Len()) + Text.Len();
 }
 
-void WriteString(FBlobWriter& Writer, const char* Text)
+HORDE_API void WriteString(FBlobWriter& Writer, const char* Text)
 {
 	WriteString(Writer, FUtf8StringView(Text));
 }
 
-void WriteString(FBlobWriter& Writer, const FUtf8StringView& Text)
+HORDE_API void WriteString(FBlobWriter& Writer, const FUtf8StringView& Text)
 {
 	WriteUnsignedVarInt(Writer, Text.Len());
 	WriteFixedLengthBytes(Writer, Text.GetData(), Text.Len());
 }
 
-void WriteString(FBlobWriter& Writer, const FUtf8String& Text)
+HORDE_API void WriteString(FBlobWriter& Writer, const FUtf8String& Text)
 {
 	WriteString(Writer, FUtf8StringView(Text));
 }
diff --git a/Engine/Source/Developer/Horde/Public/SharedBufferView.h b/Engine/Source/Developer/Horde/Public/SharedBufferView.h
--- a/Engine/Source/Developer/Horde/Public/SharedBufferView.h
+++ b/Engine/Source/Developer/Horde/Public/SharedBufferView.h
@@ -7,7 +7,7 @@
 //
 // View into a buffer with shared ownership
 //
-class FSharedBufferView final
+class HORDE_API FSharedBufferView final
 {
 public:
 	FSharedBufferView();
diff --git a/Engine/Source/Developer/Horde/Public/Storage/Blob.h b/Engine/Source/Developer/Horde/Public/Storage/Blob.h
--- a/Engine/Source/Developer/Horde/Public/Storage/Blob.h
+++ b/Engine/Source/Developer/Horde/Public/Storage/Blob.h
@@ -10,7 +10,7 @@
 /**
  * Describes a blob of data in the storage system
  */
-struct FBlob
+struct HORDE_API FBlob
 {
 	/** Type of the blob. */
 	FBlobType Type;
diff --git a/Engine/Source/Developer/Horde/Public/Storage/BlobHandle.h b/Engine/Source/Developer/Horde/Public/Storage/BlobHandle.h
--- a/Engine/Source/Developer/Horde/Public/Storage/BlobHandle.h
+++ b/Engine/Source/Developer/Horde/Public/Storage/BlobHandle.h
@@ -14,7 +14,7 @@ class FBlobLocator;
 /**
  * Handle to a blob. Can be used to reference blobs that have not been flushed yet. Implemented as a shared pointer to a FBlobHandleData object.
  */
-class FBlobHandle : public TSharedPtr<FBlobHandleData>
+class HORDE_API FBlobHandle : public TSharedPtr<FBlobHandleData>
 {
 public:
 	using Super = TSharedPtr<FBlobHandleData>;
@@ -70,7 +70,7 @@ public:
 /**
  * Base interface for a blob.
  */
-class FBlobHandleData
+class HORDE_API FBlobHandleData
 {
 public:
 	virtual ~FBlobHandleData();
@@ -115,7 +115,7 @@ public:
 /**
  * Stores a blob handle along with a hash of the target node.
  */
-struct FBlobHandleWithHash
+struct HORDE_API FBlobHandleWithHash
 {
 	FBlobHandle Handle;
 	FIoHash Hash;
diff --git a/Engine/Source/Developer/Horde/Public/Storage/BlobLocator.h b/Engine/Source/Developer/Horde/Public/Storage/BlobLocator.h
--- a/Engine/Source/Developer/Horde/Public/Storage/BlobLocator.h
+++ b/Engine/Source/Developer/Horde/Public/Storage/BlobLocator.h
@@ -8,7 +8,7 @@
 /**
  * Identifies the location of a blob. Meaning of this string is implementation defined.
  */
-class FBlobLocator
+class HORDE_API FBlobLocator
 {
 public:
 	FBlobLocator();
diff --git a/Engine/Source/Developer/Horde/Public/Storage/BlobReader.h b/Engine/Source/Developer/Horde/Public/Storage/BlobReader.h
--- a/Engine/Source/Developer/Horde/Public/Storage/BlobReader.h
+++ b/Engine/Source/Developer/Horde/Public/Storage/BlobReader.h
@@ -13,7 +13,7 @@ struct FIoHash;
 /**
  * Reads data from a blob object
  */
-class FBlobReader
+class HORDE_API FBlobReader
 {
 public:
 	FBlobReader(const FBlob& InBlob);
@@ -39,11 +39,11 @@ private:
 
 // ------------------------------------------------------------------------
 
-FBlobHandle ReadBlobHandle(FBlobReader& Reader);
-FBlobHandleWithHash ReadBlobHandleWithHash(FBlobReader& Reader);
-int ReadInt32(FBlobReader& Reader);
-FIoHash ReadIoHash(FBlobReader& Reader);
-FMemoryView ReadFixedLengthBytes(FBlobReader& Reader, size_t Length);
-size_t ReadUnsignedVarInt(FBlobReader& Reader);
-FUtf8String ReadString(FBlobReader& Reader);
-FMemoryView ReadStringSpan(FBlobReader& Reader);
+HORDE_API FBlobHandle ReadBlobHandle(FBlobReader& Reader);
+HORDE_API FBlobHandleWithHash ReadBlobHandleWithHash(FBlobReader& Reader);
+HORDE_API int ReadInt32(FBlobReader& Reader);
+HORDE_API FIoHash ReadIoHash(FBlobReader& Reader);
+HORDE_API FMemoryView ReadFixedLengthBytes(FBlobReader& Reader, size_t Length);
+HORDE_API size_t ReadUnsignedVarInt(FBlobReader& Reader);
+HORDE_API FUtf8String ReadString(FBlobReader& Reader);
+HORDE_API FMemoryView ReadStringSpan(FBlobReader& Reader);
diff --git a/Engine/Source/Developer/Horde/Public/Storage/BlobType.h b/Engine/Source/Developer/Horde/Public/Storage/BlobType.h
--- a/Engine/Source/Developer/Horde/Public/Storage/BlobType.h
+++ b/Engine/Source/Developer/Horde/Public/Storage/BlobType.h
@@ -7,7 +7,7 @@
 /**
  * Identifies the type of a blob
  */
-struct FBlobType
+struct HORDE_API FBlobType
 {
 	static const FBlobType Leaf;
 
diff --git a/Engine/Source/Developer/Horde/Public/Storage/BlobWriter.h b/Engine/Source/Developer/Horde/Public/Storage/BlobWriter.h
--- a/Engine/Source/Developer/Horde/Public/Storage/BlobWriter.h
+++ b/Engine/Source/Developer/Horde/Public/Storage/BlobWriter.h
@@ -15,7 +15,7 @@ struct FRefOptions;
 /**
  * Interface for writing data to the storage system.
  */
-class FBlobWriter
+class HORDE_API FBlobWriter
 {
 public:
 	virtual ~FBlobWriter();
@@ -51,21 +51,21 @@ public:
 
 // ------------------------------------------------------------------------------------------
 
-void WriteBlobHandle(FBlobWriter& Writer, FBlobHandle Handle);
-void WriteBlobHandleWithHash(FBlobWriter& Writer, FBlobHandleWithHash Target);
+HORDE_API void WriteBlobHandle(FBlobWriter& Writer, FBlobHandle Handle);
+HORDE_API void WriteBlobHandleWithHash(FBlobWriter& Writer, FBlobHandleWithHash Target);
 
-void WriteIoHash(FBlobWriter& Writer, const struct FIoHash& Hash);
+HORDE_API void WriteIoHash(FBlobWriter& Writer, const struct FIoHash& Hash);
 
-void WriteFixedLengthBytes(FBlobWriter& Writer, const void* Data, size_t Length);
-void WriteFixedLengthBytes(FBlobWriter& Writer, const FMemoryView& View);
+HORDE_API void WriteFixedLengthBytes(FBlobWriter& Writer, const void* Data, size_t Length);
+HORDE_API void WriteFixedLengthBytes(FBlobWriter& Writer, const FMemoryView& View);
 
-size_t MeasureUnsignedVarInt(size_t Value);
-size_t WriteUnsignedVarInt(void* Buffer, size_t Value);
-void WriteUnsignedVarIntWithKnownLength(void* Buffer, size_t Value, size_t NumBytes);
-void WriteUnsignedVarInt(FBlobWriter& Writer, size_t Value);
+HORDE_API size_t MeasureUnsignedVarInt(size_t Value);
+HORDE_API size_t WriteUnsignedVarInt(void* Buffer, size_t Value);
+HORDE_API void WriteUnsignedVarIntWithKnownLength(void* Buffer, size_t Value, size_t NumBytes);
+HORDE_API void WriteUnsignedVarInt(FBlobWriter& Writer, size_t Value);
 
-size_t MeasureString(const char* Text);
-size_t MeasureString(const FUtf8StringView& Text);
-void WriteString(FBlobWriter& Writer, const char* Text);
-void WriteString(FBlobWriter& Writer, const FUtf8StringView& Text);
-void WriteString(FBlobWriter& Writer, const FUtf8String& Text);
+HORDE_API size_t MeasureString(const char* Text);
+HORDE_API size_t MeasureString(const FUtf8StringView& Text);
+HORDE_API void WriteString(FBlobWriter& Writer, const char* Text);
+HORDE_API void WriteString(FBlobWriter& Writer, const FUtf8StringView& Text);
+HORDE_API void WriteString(FBlobWriter& Writer, const FUtf8String& Text);
diff --git a/Engine/Source/Developer/Horde/Public/Storage/Bundles/Bundle.h b/Engine/Source/Developer/Horde/Public/Storage/Bundles/Bundle.h
--- a/Engine/Source/Developer/Horde/Public/Storage/Bundles/Bundle.h
+++ b/Engine/Source/Developer/Horde/Public/Storage/Bundles/Bundle.h
@@ -41,7 +41,7 @@ enum class EBundleVersion : unsigned char
 /**
  * General bundle properties
  */
-struct FBundle
+struct HORDE_API FBundle
 {
 	static const FBlobType BlobType;
 };
@@ -49,7 +49,7 @@ struct FBundle
 /**
  * Signature for a bundle
  */
-struct FBundleSignature
+struct HORDE_API FBundleSignature
 {
 	/** Number of bytes in a serialized signature. */
 	static const int NumBytes = 8;
diff --git a/Engine/Source/Developer/Horde/Public/Storage/Bundles/BundleCompression.h b/Engine/Source/Developer/Horde/Public/Storage/Bundles/BundleCompression.h
--- a/Engine/Source/Developer/Horde/Public/Storage/Bundles/BundleCompression.h
+++ b/Engine/Source/Developer/Horde/Public/Storage/Bundles/BundleCompression.h
@@ -27,7 +27,7 @@ enum EBundleCompressionFormat : unsigned char
 /*
  * Utility methods for compressing bundles
  */
-struct FBundleCompression
+struct HORDE_API FBundleCompression
 {
 	/** Gets the maximum size of the buffer required to compress the given data. */
 	static size_t GetMaxSize(EBundleCompressionFormat Format, const FMemoryView& Input);
diff --git a/Engine/Source/Developer/Horde/Public/Storage/Bundles/BundleOptions.h b/Engine/Source/Developer/Horde/Public/Storage/Bundles/BundleOptions.h
--- a/Engine/Source/Developer/Horde/Public/Storage/Bundles/BundleOptions.h
+++ b/Engine/Source/Developer/Horde/Public/Storage/Bundles/BundleOptions.h
@@ -8,7 +8,7 @@
 /*
  * Options for configuring a bundle serializer
  */
-struct FBundleOptions
+struct HORDE_API FBundleOptions
 {
 	/** Default options value. */
 	static const FBundleOptions Default;
diff --git a/Engine/Source/Developer/Horde/Public/Storage/Bundles/V2/BundleWriter.h b/Engine/Source/Developer/Horde/Public/Storage/Bundles/V2/BundleWriter.h
--- a/Engine/Source/Developer/Horde/Public/Storage/Bundles/V2/BundleWriter.h
+++ b/Engine/Source/Developer/Horde/Public/Storage/Bundles/V2/BundleWriter.h
@@ -15,7 +15,7 @@ class FPacketWriter;
 /**
  * Writes blobs into bundles
  */
-class FBundleWriter final : public FBlobWriter
+class HORDE_API FBundleWriter final : public FBlobWriter
 {
 public:
 	FBundleWriter(TSharedRef<FKeyValueStorageClient> InStorageClient, FUtf8String InBasePath, const FBundleOptions& InOptions);
diff --git a/Engine/Source/Developer/Horde/Public/Storage/Bundles/V2/ExportHandle.h b/Engine/Source/Developer/Horde/Public/Storage/Bundles/V2/ExportHandle.h
--- a/Engine/Source/Developer/Horde/Public/Storage/Bundles/V2/ExportHandle.h
+++ b/Engine/Source/Developer/Horde/Public/Storage/Bundles/V2/ExportHandle.h
@@ -14,7 +14,7 @@ typedef TBlobHandle<FExportHandleData> FExportHandle;
 /** 
  * Implementation of export handle data.
  */
-class FExportHandleData final : public FBlobHandleData, public TSharedFromThis<FExportHandleData, ESPMode::ThreadSafe>
+class HORDE_API FExportHandleData final : public FBlobHandleData, public TSharedFromThis<FExportHandleData, ESPMode::ThreadSafe>
 {
 public:
 	static const char Type[];
diff --git a/Engine/Source/Developer/Horde/Public/Storage/Bundles/V2/Packet.h b/Engine/Source/Developer/Horde/Public/Storage/Bundles/V2/Packet.h
--- a/Engine/Source/Developer/Horde/Public/Storage/Bundles/V2/Packet.h
+++ b/Engine/Source/Developer/Horde/Public/Storage/Bundles/V2/Packet.h
@@ -15,7 +15,7 @@ struct FPacketExport;
 /**
  * Accessor for data structures stored into a serialized bundle packet.
  */
-class FPacket
+class HORDE_API FPacket
 {
 public:
 	static const FBlobType BlobType;
@@ -62,7 +62,7 @@ private:
 /**
  * Specifies the path to an imported node
  */
-struct FPacketImport
+struct HORDE_API FPacketImport
 {
 	static constexpr int32 Bias = 3;
 	static constexpr int32 InvalidBaseIdx = -1;
@@ -86,7 +86,7 @@ private:
 /*
  * Data for an exported node in a packet
  */
-struct FPacketExport
+struct HORDE_API FPacketExport
 {
 	FPacketExport(FSharedBufferView InBuffer);
 	~FPacketExport();
diff --git a/Engine/Source/Developer/Horde/Public/Storage/Bundles/V2/PacketHandle.h b/Engine/Source/Developer/Horde/Public/Storage/Bundles/V2/PacketHandle.h
--- a/Engine/Source/Developer/Horde/Public/Storage/Bundles/V2/PacketHandle.h
+++ b/Engine/Source/Developer/Horde/Public/Storage/Bundles/V2/PacketHandle.h
@@ -16,7 +16,7 @@ class FPacketHandleData;
 typedef TBlobHandle<FPacketHandleData> FPacketHandle;
 
 // Data for FPacketHandle
-class FPacketHandleData final : public FBlobHandleData, public TSharedFromThis<FPacketHandleData, ESPMode::ThreadSafe>
+class HORDE_API FPacketHandleData final : public FBlobHandleData, public TSharedFromThis<FPacketHandleData, ESPMode::ThreadSafe>
 {
 public:
 	static const char Type[];
diff --git a/Engine/Source/Developer/Horde/Public/Storage/Bundles/V2/PacketReader.h b/Engine/Source/Developer/Horde/Public/Storage/Bundles/V2/PacketReader.h
--- a/Engine/Source/Developer/Horde/Public/Storage/Bundles/V2/PacketReader.h
+++ b/Engine/Source/Developer/Horde/Public/Storage/Bundles/V2/PacketReader.h
@@ -10,7 +10,7 @@
 /**
  * Utility class for constructing BlobData objects from a packet, caching any computed handles to other blobs.
  */
-class FPacketReader
+class HORDE_API FPacketReader
 {
 public:
 	FPacketReader(TSharedRef<FStorageClient> InStorageClient, FBlobHandle InBundleHandle, FPacketHandle InPacketHandle, FPacket InPacket);
diff --git a/Engine/Source/Developer/Horde/Public/Storage/Bundles/V2/PacketWriter.h b/Engine/Source/Developer/Horde/Public/Storage/Bundles/V2/PacketWriter.h
--- a/Engine/Source/Developer/Horde/Public/Storage/Bundles/V2/PacketWriter.h
+++ b/Engine/Source/Developer/Horde/Public/Storage/Bundles/V2/PacketWriter.h
@@ -11,7 +11,7 @@
 /**
  * Writes exports into a new bundle packet
  */
-class FPacketWriter
+class HORDE_API FPacketWriter
 {
 public:
 	FPacketWriter(FBlobHandle InBundleHandle, FBlobHandle InPacketHandle);
diff --git a/Engine/Source/Developer/Horde/Public/Storage/ChunkedBufferWriter.h b/Engine/Source/Developer/Horde/Public/Storage/ChunkedBufferWriter.h
--- a/Engine/Source/Developer/Horde/Public/Storage/ChunkedBufferWriter.h
+++ b/Engine/Source/Developer/Horde/Public/Storage/ChunkedBufferWriter.h
@@ -8,11 +8,12 @@
 /**
  * Writes data to a series of memory blocks
  */
-class FChunkedBufferWriter
+class HORDE_API FChunkedBufferWriter
 {
 public:
 	FChunkedBufferWriter(size_t InitialSize = 1024);
 	FChunkedBufferWriter(const FChunkedBufferWriter&) = delete;
+	FChunkedBufferWriter& operator = (const FChunkedBufferWriter&) = delete;
 	virtual ~FChunkedBufferWriter();
 
 	/** Reset the contents of this writer. */
diff --git a/Engine/Source/Developer/Horde/Public/Storage/Clients/BundleStorageClient.h b/Engine/Source/Developer/Horde/Public/Storage/Clients/BundleStorageClient.h
--- a/Engine/Source/Developer/Horde/Public/Storage/Clients/BundleStorageClient.h
+++ b/Engine/Source/Developer/Horde/Public/Storage/Clients/BundleStorageClient.h
@@ -11,7 +11,7 @@
 /**
  * Base class for an implementation of <see cref="IStorageClient"/>, providing implementations for some common functionality using bundles.
  */
-class FBundleStorageClient final : public FKeyValueStorageClient
+class HORDE_API FBundleStorageClient final : public FKeyValueStorageClient
 {
 public:
 	FBundleStorageClient(TSharedRef<FKeyValueStorageClient> InInner);
diff --git a/Engine/Source/Developer/Horde/Public/Storage/Clients/FileStorageClient.h b/Engine/Source/Developer/Horde/Public/Storage/Clients/FileStorageClient.h
--- a/Engine/Source/Developer/Horde/Public/Storage/Clients/FileStorageClient.h
+++ b/Engine/Source/Developer/Horde/Public/Storage/Clients/FileStorageClient.h
@@ -9,7 +9,7 @@
 /**
  * Implementation of FStorageClient which writes data to files on disk.
  */
-class FFileStorageClient : public FKeyValueStorageClient
+class HORDE_API FFileStorageClient : public FKeyValueStorageClient
 {
 public:
 	FFileStorageClient(std::filesystem::path InRootDir);
diff --git a/Engine/Source/Developer/Horde/Public/Storage/Clients/KeyValueStorageClient.h b/Engine/Source/Developer/Horde/Public/Storage/Clients/KeyValueStorageClient.h
--- a/Engine/Source/Developer/Horde/Public/Storage/Clients/KeyValueStorageClient.h
+++ b/Engine/Source/Developer/Horde/Public/Storage/Clients/KeyValueStorageClient.h
@@ -9,7 +9,7 @@
 /*
  * Base class for storage clients that wrap a diirect key/value type store without any merging/splitting.
  */
-class FKeyValueStorageClient : public FStorageClient
+class HORDE_API FKeyValueStorageClient : public FStorageClient
 {
 public:
 	/** Read a single blob from the underlying store. */
diff --git a/Engine/Source/Developer/Horde/Public/Storage/Clients/MemoryStorageClient.h b/Engine/Source/Developer/Horde/Public/Storage/Clients/MemoryStorageClient.h
--- a/Engine/Source/Developer/Horde/Public/Storage/Clients/MemoryStorageClient.h
+++ b/Engine/Source/Developer/Horde/Public/Storage/Clients/MemoryStorageClient.h
@@ -8,7 +8,7 @@
 /**
  * Implementation of FStorageClient which stores data in memory.
  */
-class FMemoryStorageClient final : public FKeyValueStorageClient
+class HORDE_API FMemoryStorageClient final : public FKeyValueStorageClient
 {
 public:
 	FMemoryStorageClient();
diff --git a/Engine/Source/Developer/Horde/Public/Storage/Nodes/ChunkNode.h b/Engine/Source/Developer/Horde/Public/Storage/Nodes/ChunkNode.h
--- a/Engine/Source/Developer/Horde/Public/Storage/Nodes/ChunkNode.h
+++ b/Engine/Source/Developer/Horde/Public/Storage/Nodes/ChunkNode.h
@@ -14,7 +14,7 @@
 /**
  * Options for chunking data
  */
-struct FChunkingOptions
+struct HORDE_API FChunkingOptions
 {
 	static const FChunkingOptions Default;
 
@@ -26,7 +26,7 @@ struct FChunkingOptions
 /**
  * Node containing a chunk of data
  */
-class FChunkNode
+class HORDE_API FChunkNode
 {
 public:
 	static const FBlobType LeafBlobType;
@@ -48,7 +48,7 @@ public:
 /**
  * Utility class for reading data a data stream from a tree of chunk nodes
  */
-class FChunkNodeReader
+class HORDE_API FChunkNodeReader
 {
 public:
 	FChunkNodeReader(FBlob Blob);
@@ -69,7 +69,7 @@ private:
 /**
  * Utility class for writing new data to a tree of chunk nodes
  */
-class FChunkNodeWriter
+class HORDE_API FChunkNodeWriter
 {
 public:
 	FChunkNodeWriter(FBlobWriter& InWriter, const FChunkingOptions& InOptions = FChunkingOptions::Default);
diff --git a/Engine/Source/Developer/Horde/Public/Storage/Nodes/DirectoryEntry.h b/Engine/Source/Developer/Horde/Public/Storage/Nodes/DirectoryEntry.h
--- a/Engine/Source/Developer/Horde/Public/Storage/Nodes/DirectoryEntry.h
+++ b/Engine/Source/Developer/Horde/Public/Storage/Nodes/DirectoryEntry.h
@@ -13,7 +13,7 @@ class FBlobWriter;
 /**
  * Entry for a directory within a directory node
  */
-class FDirectoryEntry
+class HORDE_API FDirectoryEntry
 {
 public:
 	FBlobHandle Target;
diff --git a/Engine/Source/Developer/Horde/Public/Storage/Nodes/DirectoryNode.h b/Engine/Source/Developer/Horde/Public/Storage/Nodes/DirectoryNode.h
--- a/Engine/Source/Developer/Horde/Public/Storage/Nodes/DirectoryNode.h
+++ b/Engine/Source/Developer/Horde/Public/Storage/Nodes/DirectoryNode.h
@@ -20,7 +20,7 @@ enum class EDirectoryFlags
 /**
  * Stores the contents of a directory in a blob
  */
-class FDirectoryNode
+class HORDE_API FDirectoryNode
 {
 public:
 	static const FBlobType BlobType;
diff --git a/Engine/Source/Developer/Horde/Public/Storage/Nodes/FileEntry.h b/Engine/Source/Developer/Horde/Public/Storage/Nodes/FileEntry.h
--- a/Engine/Source/Developer/Horde/Public/Storage/Nodes/FileEntry.h
+++ b/Engine/Source/Developer/Horde/Public/Storage/Nodes/FileEntry.h
@@ -40,7 +40,7 @@ ENUM_CLASS_FLAGS(EFileEntryFlags)
 /**
  * Entry for a file within a directory node
  */
-class FFileEntry
+class HORDE_API FFileEntry
 {
 public:
 	/** Handle to the root chunked data node. */
diff --git a/Engine/Source/Developer/Horde/Public/Storage/RefName.h b/Engine/Source/Developer/Horde/Public/Storage/RefName.h
--- a/Engine/Source/Developer/Horde/Public/Storage/RefName.h
+++ b/Engine/Source/Developer/Horde/Public/Storage/RefName.h
@@ -6,7 +6,7 @@
 #include "Containers/Utf8String.h"
 
 // Identifier for a ref in the storage system. Refs serve as GC roots, and are persistent entry points to expanding data structures within the store.
-struct FRefName
+struct HORDE_API FRefName
 {
 public:
 	FRefName(FUtf8String Text);
diff --git a/Engine/Source/Developer/Horde/Public/Storage/RefOptions.h b/Engine/Source/Developer/Horde/Public/Storage/RefOptions.h
--- a/Engine/Source/Developer/Horde/Public/Storage/RefOptions.h
+++ b/Engine/Source/Developer/Horde/Public/Storage/RefOptions.h
@@ -7,7 +7,7 @@
 /**
  * Options for adding a new ref
  */ 
-struct FRefOptions
+struct HORDE_API FRefOptions
 {
 	static const FRefOptions Default;
 
diff --git a/Engine/Source/Developer/Horde/Public/Storage/StorageClient.h b/Engine/Source/Developer/Horde/Public/Storage/StorageClient.h
--- a/Engine/Source/Developer/Horde/Public/Storage/StorageClient.h
+++ b/Engine/Source/Developer/Horde/Public/Storage/StorageClient.h
@@ -29,7 +29,7 @@ struct FRefCacheTime
 /**
  * Interface for the storage system.
  */
-class FStorageClient : public TSharedFromThis<FStorageClient, ESPMode::ThreadSafe>
+class HORDE_API FStorageClient : public TSharedFromThis<FStorageClient, ESPMode::ThreadSafe>
 {
 public:
 	virtual ~FStorageClient();
