Horde: Add a module declaration for the Horde module, to fix allocator crashes on objects passed over the DLL boundary.

diff --git a/Engine/Source/Developer/Horde/Private/Horde.cpp b/Engine/Source/Developer/Horde/Private/Horde.cpp
--- a/Engine/Source/Developer/Horde/Private/Horde.cpp
+++ b/Engine/Source/Developer/Horde/Private/Horde.cpp
@@ -3,6 +3,19 @@
 #include "Horde.h"
 #include "Misc/CommandLine.h"
 #include "Misc/Parse.h"
+#include "Modules/ModuleManager.h"
+#include "Modules/ModuleInterface.h"
+
+struct FHordeModule : IModuleInterface
+{
+	virtual void StartupModule() override
+	{
+	}
+};
+
+IMPLEMENT_MODULE(FHordeModule, Horde)
+
+// --------------------------------------------------------------------------------
 
 FString FHorde::GetTemplateName()
 {
