Horde: Remove hard-coded Horde URL from code exposed to licensees.

diff --git a/Engine/Source/Developer/Horde/Private/Horde.cpp b/Engine/Source/Developer/Horde/Private/Horde.cpp
--- a/Engine/Source/Developer/Horde/Private/Horde.cpp
+++ b/Engine/Source/Developer/Horde/Private/Horde.cpp
@@ -63,7 +63,25 @@ FString FHorde::GetJobId()
 
 FString FHorde::GetJobURL()
 {
-	return FString::Printf(TEXT("https://horde.devtools.epicgames.com/job/%s"), *GetJobId());
+	struct FJobUrlInitializer
+	{
+		FString Value;
+
+		FJobUrlInitializer()
+		{
+			if (!FParse::Value(FCommandLine::Get(), TEXT("HordeJobUrl="), Value))
+			{
+				Value = FPlatformMisc::GetEnvironmentVariable(TEXT("UE_HORDE_URL"));
+				if (Value.Len() > 0)
+				{
+					Value /= FString::Printf(TEXT("job/%s"), *GetJobId());
+				}
+			}
+		}
+	};
+
+	static const FJobUrlInitializer JobUrl;
+	return JobUrl.Value;
 }
 
 FString FHorde::GetStepId()
