Horde: Fix incompatibility between C# and C++ chunk node serialization.

diff --git a/Engine/Source/Developer/Horde/Private/Storage/Nodes/ChunkNode.cpp b/Engine/Source/Developer/Horde/Private/Storage/Nodes/ChunkNode.cpp
--- a/Engine/Source/Developer/Horde/Private/Storage/Nodes/ChunkNode.cpp
+++ b/Engine/Source/Developer/Horde/Private/Storage/Nodes/ChunkNode.cpp
@@ -28,10 +28,14 @@ FChunkNode FChunkNode::Read(FBlob Blob)
 {
 	FChunkNode Node;
 
-	FIoHash* Hashes = (FIoHash*)Blob.Data.GetPointer();
+	const uint8* InputData = Blob.Data.GetPointer();
 	for (int32 Idx = 0; Idx < Blob.References.Num(); Idx++)
 	{
-		Node.Children.Add(FBlobHandleWithHash(MoveTemp(Blob.References[Idx]), Hashes[Idx]));
+		FIoHash Hash;
+		memcpy(&Hash, InputData, sizeof(FIoHash));
+		InputData += sizeof(FIoHash) + 1; // Ignore node type
+
+		Node.Children.Add(FBlobHandleWithHash(MoveTemp(Blob.References[Idx]), Hash));
 	}
 
 	Node.Data = Blob.Data.Slice(sizeof(FIoHash) * Blob.References.Num());
@@ -45,16 +49,21 @@ FBlobHandleWithHash FChunkNode::Write(FBlobWriter& Writer) const
 
 FBlobHandleWithHash FChunkNode::Write(FBlobWriter& Writer, const TArrayView<const FBlobHandleWithHash>& Children, FMemoryView Data)
 {
-	int32 BufferSize = (sizeof(FIoHash) * Children.Num()) + Data.GetSize();
+	int32 BufferSize = ((sizeof(FIoHash) + 1) * Children.Num()) + Data.GetSize();
 	uint8* Buffer = (uint8*)Writer.GetOutputBuffer(BufferSize);
 
-	FIoHash* NextHash = (FIoHash*)Buffer;
+	uint8* NextData = (uint8*)Buffer;
 	for (int32 Idx = 0; Idx < Children.Num(); Idx++)
 	{
-		*(NextHash++) = Children[Idx].Hash;
+		memcpy(NextData, &Children[Idx].Hash, sizeof(FIoHash));
+		NextData += sizeof(FIoHash);
+
+		*NextData = 1; // Leaf node
+		NextData++;
+
 		Writer.AddImport(Children[Idx].Handle);
 	}
-	memcpy(NextHash, Data.GetData(), Data.GetSize());
+	memcpy(NextData, Data.GetData(), Data.GetSize());
 
 	FIoHash Hash = FIoHash::HashBuffer(Buffer, BufferSize);
 	Writer.Advance(BufferSize);
