Return EAgentMessageType::None instead of asserting null-pointer when reading a response to detect end of stream.

#rnx

diff --git a/Engine/Source/Developer/Horde/Private/Compute/AgentMessage.cpp b/Engine/Source/Developer/Horde/Private/Compute/AgentMessage.cpp
--- a/Engine/Source/Developer/Horde/Private/Compute/AgentMessage.cpp
+++ b/Engine/Source/Developer/Horde/Private/Compute/AgentMessage.cpp
@@ -163,7 +163,10 @@ EAgentMessageType FAgentMessageChannel::ReadResponse()
 	}
 
 	const unsigned char* Header = ChannelBuffers->Reader.WaitToRead(MessageHeaderLength);
-	check(Header != nullptr);
+	if (Header == nullptr)
+	{
+		return EAgentMessageType::None;
+	}
 	unsigned int Length = *((unsigned int*)(Header + 1));
 
 	Header = ChannelBuffers->Reader.WaitToRead(MessageHeaderLength + Length);
