Don't add/subtract header length in FAgentMessageChannel::GetResponseSize() and GetResponseData() as they already contain the actual values.

#rnx

diff --git a/Engine/Source/Developer/Horde/Public/Compute/AgentMessage.h b/Engine/Source/Developer/Horde/Public/Compute/AgentMessage.h
--- a/Engine/Source/Developer/Horde/Public/Compute/AgentMessage.h
+++ b/Engine/Source/Developer/Horde/Public/Compute/AgentMessage.h
@@ -134,10 +134,10 @@ public:
 	HORDE_API EAgentMessageType ReadResponse();
 
 	// Gets the raw data from a response
-	HORDE_API const void* GetResponseData() const { return ResponseData + MessageHeaderLength; }
+	HORDE_API const void* GetResponseData() const { return ResponseData; }
 
 	// Gets the size of the response
-	HORDE_API size_t GetResponseSize() const { return ResponseLength - MessageHeaderLength; }
+	HORDE_API size_t GetResponseSize() const { return ResponseLength; }
 
 	// Reads an exception response
 	HORDE_API void ReadException(AgentMessage::FException& Ex);
