Added timeout value to ReadResponse() to avoid hanging when reading feedback from Horde agent.

#rnx

diff --git a/Engine/Source/Developer/Horde/Private/Compute/AgentMessage.cpp b/Engine/Source/Developer/Horde/Private/Compute/AgentMessage.cpp
--- a/Engine/Source/Developer/Horde/Private/Compute/AgentMessage.cpp
+++ b/Engine/Source/Developer/Horde/Private/Compute/AgentMessage.cpp
@@ -153,7 +153,7 @@ void FAgentMessageChannel::Xor(const unsigned char* Data, size_t Length, unsigne
 	FlushMessage();
 }
 
-EAgentMessageType FAgentMessageChannel::ReadResponse()
+EAgentMessageType FAgentMessageChannel::ReadResponse(int32 TimeoutMS)
 {
 	if (ResponseData)
 	{
@@ -162,7 +162,7 @@ EAgentMessageType FAgentMessageChannel::ReadResponse()
 		ResponseLength = 0;
 	}
 
-	const unsigned char* Header = ChannelBuffers->Reader.WaitToRead(MessageHeaderLength);
+	const unsigned char* Header = ChannelBuffers->Reader.WaitToRead(MessageHeaderLength, TimeoutMS);
 	if (Header == nullptr)
 	{
 		return EAgentMessageType::None;
diff --git a/Engine/Source/Developer/Horde/Public/Compute/AgentMessage.h b/Engine/Source/Developer/Horde/Public/Compute/AgentMessage.h
--- a/Engine/Source/Developer/Horde/Public/Compute/AgentMessage.h
+++ b/Engine/Source/Developer/Horde/Public/Compute/AgentMessage.h
@@ -131,7 +131,7 @@ public:
 	//// Responses ////
 
 	// Reads a response from the remote. Other Read methods can be used to access response data.
-	HORDE_API EAgentMessageType ReadResponse();
+	HORDE_API EAgentMessageType ReadResponse(int32 TimeoutMS = -1);
 
 	// Gets the raw data from a response
 	HORDE_API const void* GetResponseData() const { return ResponseData; }
