Horde: Skip over target node type in file entry objects.

diff --git a/Engine/Source/Developer/Horde/Private/Storage/BlobReader.cpp b/Engine/Source/Developer/Horde/Private/Storage/BlobReader.cpp
--- a/Engine/Source/Developer/Horde/Private/Storage/BlobReader.cpp
+++ b/Engine/Source/Developer/Horde/Private/Storage/BlobReader.cpp
@@ -8,17 +8,23 @@
 // ------------------------------------------------------------------------
 
 FBlobReader::FBlobReader(const FBlob& InBlob)
-	: FBlobReader(InBlob.Data.GetView(), InBlob.References)
+	: FBlobReader(InBlob.Type.Version, InBlob.Data.GetView(), InBlob.References)
 {
 }
 
-FBlobReader::FBlobReader(const FMemoryView& InBuffer, const TArray<FBlobHandle>& InImports)
-	: Buffer(InBuffer)
+FBlobReader::FBlobReader(int32 InVersion, const FMemoryView& InBuffer, const TArray<FBlobHandle>& InImports)
+	: Version(InVersion)
+	, Buffer(InBuffer)
 	, Imports(InImports)
 	, NextImportIdx(0)
 {
 }
 
+int32 FBlobReader::GetVersion() const
+{
+	return Version;
+}
+
 const unsigned char* FBlobReader::GetBuffer() const
 {
 	return (const unsigned char*)GetView().GetData();
diff --git a/Engine/Source/Developer/Horde/Private/Storage/Nodes/FileEntry.cpp b/Engine/Source/Developer/Horde/Private/Storage/Nodes/FileEntry.cpp
--- a/Engine/Source/Developer/Horde/Private/Storage/Nodes/FileEntry.cpp
+++ b/Engine/Source/Developer/Horde/Private/Storage/Nodes/FileEntry.cpp
@@ -21,6 +21,10 @@ FFileEntry::~FFileEntry()
 FFileEntry FFileEntry::Read(FBlobReader& Reader)
 {
 	FBlobHandleWithHash Target = ReadBlobHandleWithHash(Reader);
+	if (Reader.GetVersion() >= 2)
+	{
+		ReadUnsignedVarInt(Reader);
+	}
 
 	FUtf8String Name = ReadString(Reader);
 	EFileEntryFlags Flags = (EFileEntryFlags)ReadUnsignedVarInt(Reader);
diff --git a/Engine/Source/Developer/Horde/Public/Storage/BlobReader.h b/Engine/Source/Developer/Horde/Public/Storage/BlobReader.h
--- a/Engine/Source/Developer/Horde/Public/Storage/BlobReader.h
+++ b/Engine/Source/Developer/Horde/Public/Storage/BlobReader.h
@@ -17,7 +17,10 @@ class HORDE_API FBlobReader
 {
 public:
 	FBlobReader(const FBlob& InBlob);
-	FBlobReader(const FMemoryView& InBuffer, const TArray<FBlobHandle>& InImports);
+	FBlobReader(int32 InVersion, const FMemoryView& InBuffer, const TArray<FBlobHandle>& InImports);
+
+	/** Gets a version number for the current blob */
+	int GetVersion() const;
 
 	/** Gets a pointer to the remaining memory. */
 	const unsigned char* GetBuffer() const;
@@ -32,6 +35,7 @@ public:
 	FBlobHandle ReadImport();
 
 private:
+	int32 Version;
 	FMemoryView Buffer;
 	const TArray<FBlobHandle>& Imports;
 	int32 NextImportIdx;
