[Horde]
* Fixed so code can compile without exceptions (code path only used for some uba testing)

diff --git a/Engine/Source/Developer/Horde/Private/HordePlatform.cpp b/Engine/Source/Developer/Horde/Private/HordePlatform.cpp
--- a/Engine/Source/Developer/Horde/Private/HordePlatform.cpp
+++ b/Engine/Source/Developer/Horde/Private/HordePlatform.cpp
@@ -41,12 +41,20 @@
 
 void FHordePlatform::NotImplemented()
 {
+#if PLATFORM_EXCEPTIONS_DISABLED
+	abort();
+#else
 	throw std::string("Not Implemented");
+#endif
 }
 
 void FHordePlatform::NotSupported(const char* Message)
 {
+#if PLATFORM_EXCEPTIONS_DISABLED
+	abort();
+#else
 	throw std::string(Message);
+#endif
 }
 
 bool FHordePlatform::GetEnvironmentVariable(const char* Name, char* Buffer, size_t BufferLen)
