[Horde]
* Fixed so semaphore names do not get too long when concatenating pid, time, counter etc
* Added asserts inside functions that never should fail

diff --git a/Engine/Source/Developer/Horde/Private/Compute/ComputeBuffer.cpp b/Engine/Source/Developer/Horde/Private/Compute/ComputeBuffer.cpp
--- a/Engine/Source/Developer/Horde/Private/Compute/ComputeBuffer.cpp
+++ b/Engine/Source/Developer/Horde/Private/Compute/ComputeBuffer.cpp
@@ -354,12 +354,14 @@ struct FComputeBufferDetail
 		TUniquePtr<FComputeBufferDetail> Detail = MakeUnique<FComputeBufferDetail>(Name);
 		if (!Detail->MemoryMappedFile.Create(NameBuffer, Capacity))
 		{
+			checkf(false, TEXT("Failed to create memory mapped file %hs"), NameBuffer);
 			return nullptr;
 		}
 
 		void* Pointer = Detail->MemoryMappedFile.GetPointer();
 		if (Pointer == nullptr)
 		{
+			checkf(false, TEXT("Failed to get pointer to memory mapping %hs"), NameBuffer);
 			return nullptr;
 		}
 
@@ -368,6 +370,7 @@ struct FComputeBufferDetail
 		snprintf(NameBuffer, FComputeBuffer::MaxNameLength, "%s_W", Name);
 		if (!Detail->WriterEvent.Create(NameBuffer))
 		{
+			checkf(false, TEXT("Failed to create writer event %hs in memory mapping"), NameBuffer);
 			return nullptr;
 		}
 
@@ -376,6 +379,7 @@ struct FComputeBufferDetail
 			snprintf(NameBuffer, FComputeBuffer::MaxNameLength, "%s_R%d", Name, ReaderIdx);
 			if (!Detail->ReaderEvents[ReaderIdx].Create(NameBuffer))
 			{
+				checkf(false, TEXT("Failed to create reader event %hs in memory mapping"), NameBuffer);
 				return nullptr;
 			}
 		}
diff --git a/Engine/Source/Developer/Horde/Private/Compute/ComputePlatform.cpp b/Engine/Source/Developer/Horde/Private/Compute/ComputePlatform.cpp
--- a/Engine/Source/Developer/Horde/Private/Compute/ComputePlatform.cpp
+++ b/Engine/Source/Developer/Horde/Private/Compute/ComputePlatform.cpp
@@ -73,6 +73,7 @@ bool FComputeEvent::Create(const char* Name)
 		Handle = Value;
 		return true;
 	}
+	checkf(false, TEXT("sem_open failed (%hs)"), strerror(errno));
 	return false;
 #endif
 }
diff --git a/Engine/Source/Developer/Horde/Private/Compute/ComputeSocket.cpp b/Engine/Source/Developer/Horde/Private/Compute/ComputeSocket.cpp
--- a/Engine/Source/Developer/Horde/Private/Compute/ComputeSocket.cpp
+++ b/Engine/Source/Developer/Horde/Private/Compute/ComputeSocket.cpp
@@ -25,13 +25,13 @@ TSharedPtr<FComputeChannel> FComputeSocket::CreateChannel(int ChannelId)
 	FComputeBuffer RecvBuffer;
 	if (!RecvBuffer.CreateNew(FComputeBuffer::FParams()))
 	{
-		return MakeShared<FComputeChannel>(FComputeChannel());
+		return MakeShared<FComputeChannel>();
 	}
 
 	FComputeBuffer SendBuffer;
 	if (!SendBuffer.CreateNew(FComputeBuffer::FParams()))
 	{
-		return MakeShared<FComputeChannel>(FComputeChannel());
+		return MakeShared<FComputeChannel>();
 	}
 
 	return CreateChannel(ChannelId, std::move(RecvBuffer), std::move(SendBuffer));
diff --git a/Engine/Source/Developer/Horde/Private/HordePlatform.cpp b/Engine/Source/Developer/Horde/Private/HordePlatform.cpp
--- a/Engine/Source/Developer/Horde/Private/HordePlatform.cpp
+++ b/Engine/Source/Developer/Horde/Private/HordePlatform.cpp
@@ -84,7 +84,7 @@ void FHordePlatform::CreateUniqueIdentifier(char* NameBuffer, size_t NameBufferL
 #else
 	struct timespec ts;
 	clock_gettime(CLOCK_MONOTONIC, &ts);
-	snprintf(NameBuffer, NameBufferLen, "%u%zu%zu_%d", getpid(), (size_t)ts.tv_sec, (size_t)ts.tv_nsec, FPlatformAtomics::InterlockedIncrement(&Counter));
+	snprintf(NameBuffer, NameBufferLen, "%u%zu%zu_%d", getpid(), (size_t)(ts.tv_sec%100000), (size_t)ts.tv_nsec, FPlatformAtomics::InterlockedIncrement(&Counter));
 #endif
 }
 
