[Horde]
* Changed ComputeSocket to use ue's FCriticalSection instead of std::lock. This to work around backward compatibility with older crt

diff --git a/Engine/Source/Developer/Horde/Private/Compute/ComputeSocket.cpp b/Engine/Source/Developer/Horde/Private/Compute/ComputeSocket.cpp
--- a/Engine/Source/Developer/Horde/Private/Compute/ComputeSocket.cpp
+++ b/Engine/Source/Developer/Horde/Private/Compute/ComputeSocket.cpp
@@ -2,13 +2,15 @@
 
 #include "Compute/ComputeSocket.h"
 #include "Compute/ComputePlatform.h"
+#include "HAL/CriticalSection.h"
+#include "HAL/Event.h"
+#include "Misc/ScopeLock.h"
 #include <iostream>
 #include <assert.h>
 #include <unordered_set>
 #include <unordered_map>
 #include <vector>
 #include <thread>
-#include <mutex>
 #include <chrono>
 #include "../HordePlatform.h"
 
@@ -263,11 +265,9 @@ public:
 
 	TUniquePtr<FComputeTransport> Transport;
 	const EComputeSocketEndpoint Endpoint;
-	std::mutex CriticalSection;
+	FCriticalSection CriticalSection;
 
-	bool bPingThreadFinish;
-	std::mutex PingThreadFinishMutex;
-	std::condition_variable PingThreadFinishCV;
+	FEventRef PingThreadFinishCV;
 	std::thread PingThread;
 
 	std::thread RecvThread;
@@ -279,19 +279,13 @@ public:
 	FRemoteComputeSocket(TUniquePtr<FComputeTransport> InTransport, EComputeSocketEndpoint InEndpoint)
 		: Transport(MoveTemp(InTransport))
 		, Endpoint(InEndpoint)
-		, CriticalSection()
-		, bPingThreadFinish(0)
+		, PingThreadFinishCV(EEventMode::ManualReset)
 	{
 	}
 
 	~FRemoteComputeSocket() override
 	{
-		{
-			std::lock_guard<std::mutex>	Lock(PingThreadFinishMutex);
-			bPingThreadFinish = 1;
-		}
-
-		PingThreadFinishCV.notify_all();
+		PingThreadFinishCV->Trigger();
 
 		for (FComputeBufferReader& Reader : Readers)
 		{
@@ -310,8 +304,6 @@ public:
 
 	virtual void StartCommunication()
 	{
-		bPingThreadFinish = 0;
-
 		// Initialize the receiver thread after having attached channel 0
 		RecvThread = std::thread(&FRemoteComputeSocket::RecvThreadProc, this);
 		PingThread = std::thread(&FRemoteComputeSocket::PingThreadProc, this);
@@ -322,7 +314,7 @@ public:
 		for (;;)
 		{
 			{ // Send the ping message
-				std::lock_guard<std::mutex> Lock(CriticalSection);
+				FScopeLock Lock(&CriticalSection);
 
 				FFrameHeader Header;
 				Header.Channel = 0;
@@ -331,8 +323,7 @@ public:
 				Transport->SendMessage(&Header, sizeof(Header));
 			}
 
-			std::unique_lock<std::mutex> Lock(PingThreadFinishMutex);
-			if (PingThreadFinishCV.wait_for(Lock, std::chrono::seconds(2), [this]() { return bPingThreadFinish; }))
+			if (PingThreadFinishCV->Wait(2000))
 			{
 				break;
 			}
@@ -370,7 +361,7 @@ public:
 		const unsigned char* Data;
 		while ((Data = Reader.WaitToRead(1)) != nullptr)
 		{
-			std::lock_guard<std::mutex> Lock(CriticalSection);
+			FScopeLock Lock(&CriticalSection);
 			Header.Size = (int)Reader.GetMaxReadSize();
 			Transport->SendMessage(&Header, sizeof(Header));
 			Transport->SendMessage(Data, Header.Size);
@@ -379,7 +370,7 @@ public:
 
 		if (Reader.IsComplete())
 		{
-			std::lock_guard<std::mutex> Lock(CriticalSection);
+			FScopeLock Lock(&CriticalSection);
 			Header.Size = (int)EControlMessageType::Detach;
 			Transport->SendMessage(&Header, sizeof(Header));
 		}
@@ -390,7 +381,7 @@ public:
 		std::unordered_map<int, FComputeBufferWriter>::iterator Iter = CachedWriters.find(Channel);
 		if (Iter == CachedWriters.end())
 		{
-			std::lock_guard<std::mutex> Lock(CriticalSection);
+			FScopeLock Lock(&CriticalSection);
 
 			Iter = Writers.find(Channel);
 			if (Iter == Writers.end())
@@ -415,14 +406,14 @@ public:
 
 	void AttachRecvBuffer(int ChannelId, FComputeBuffer RecvBuffer) override
 	{
-		std::lock_guard<std::mutex> Lock(CriticalSection);
+		FScopeLock Lock(&CriticalSection);
 		FComputeBufferWriter Writer = RecvBuffer.CreateWriter();
 		Writers.insert(std::pair<int, FComputeBufferWriter>(ChannelId, std::move(Writer)));
 	}
 
 	void AttachSendBuffer(int ChannelId, FComputeBuffer SendBuffer) override
 	{
-		std::lock_guard<std::mutex> Lock(CriticalSection);
+		FScopeLock Lock(&CriticalSection);
 		FComputeBufferReader Reader = SendBuffer.CreateReader();
 		Readers.push_back(Reader);
 		SendThreads.insert(std::make_pair(ChannelId, std::thread(&FRemoteComputeSocket::SendThreadProc, this, ChannelId, std::move(Reader))));
@@ -432,7 +423,7 @@ public:
 	{
 		CachedWriters.erase(Channel);
 
-		std::lock_guard<std::mutex> Lock(CriticalSection);
+		FScopeLock Lock(&CriticalSection);
 
 		std::unordered_map<int, FComputeBufferWriter>::iterator Iter = Writers.find(Channel);
 		if (Iter != Writers.end())
