[Horde]
* Fixed include paths to not use ../..

diff --git a/Engine/Source/Developer/Horde/Private/Storage/Bundles/BundleCompression.cpp b/Engine/Source/Developer/Horde/Private/Storage/Bundles/BundleCompression.cpp
--- a/Engine/Source/Developer/Horde/Private/Storage/Bundles/BundleCompression.cpp
+++ b/Engine/Source/Developer/Horde/Private/Storage/Bundles/BundleCompression.cpp
@@ -1,8 +1,8 @@
 // Copyright Epic Games, Inc. All Rights Reserved.
 
 #include "Storage/Bundles/BundleCompression.h"
+#include "HordePlatform.h"
 #include "Misc/Compression.h"
-#include "../../HordePlatform.h"
 
 size_t FBundleCompression::GetMaxSize(EBundleCompressionFormat Format, const FMemoryView& Input)
 {
diff --git a/Engine/Source/Developer/Horde/Private/Storage/Bundles/V2/BundleWriter.cpp b/Engine/Source/Developer/Horde/Private/Storage/Bundles/V2/BundleWriter.cpp
--- a/Engine/Source/Developer/Horde/Private/Storage/Bundles/V2/BundleWriter.cpp
+++ b/Engine/Source/Developer/Horde/Private/Storage/Bundles/V2/BundleWriter.cpp
@@ -1,11 +1,11 @@
 // Copyright Epic Games, Inc. All Rights Reserved.
 
 #include "Storage/Bundles/V2/BundleWriter.h"
+#include "HordePlatform.h"
 #include "Storage/Bundles/V2/PacketWriter.h"
 #include "Storage/Blob.h"
 #include "Storage/StorageClient.h"
 #include "Storage/Bundles/V2/ExportHandle.h"
-#include "../../../HordePlatform.h"
 #include "Storage/ChunkedBufferWriter.h"
 
 //
diff --git a/Engine/Source/Developer/Horde/Private/Storage/Bundles/V2/ExportHandle.cpp b/Engine/Source/Developer/Horde/Private/Storage/Bundles/V2/ExportHandle.cpp
--- a/Engine/Source/Developer/Horde/Private/Storage/Bundles/V2/ExportHandle.cpp
+++ b/Engine/Source/Developer/Horde/Private/Storage/Bundles/V2/ExportHandle.cpp
@@ -1,7 +1,7 @@
 // Copyright Epic Games, Inc. All Rights Reserved.
 
 #include "Storage/Bundles/V2/ExportHandle.h"
-#include "../../../HordePlatform.h"
+#include "HordePlatform.h"
 
 const char FExportHandleData::Type[] = "FExportHandleData";
 const FUtf8StringView FExportHandleData::FragmentPrefix("exp=");
diff --git a/Engine/Source/Developer/Horde/Private/Storage/Bundles/V2/Packet.cpp b/Engine/Source/Developer/Horde/Private/Storage/Bundles/V2/Packet.cpp
--- a/Engine/Source/Developer/Horde/Private/Storage/Bundles/V2/Packet.cpp
+++ b/Engine/Source/Developer/Horde/Private/Storage/Bundles/V2/Packet.cpp
@@ -4,8 +4,8 @@
 #include "Storage/Bundles/Bundle.h"
 #include "Storage/BlobWriter.h"
 #include "Storage/ChunkedBufferWriter.h"
-#include "../../../HordePlatform.h"
 #include "Serialization/VarInt.h"
+#include "HordePlatform.h"
 
 const FBlobType FPacket::BlobType(FGuid(0xCD9A04EF, 0x47D3CAC1, 0x492A05A6, 0x51E63081), 1);
 
diff --git a/Engine/Source/Developer/Horde/Private/Storage/Bundles/V2/PacketHandle.cpp b/Engine/Source/Developer/Horde/Private/Storage/Bundles/V2/PacketHandle.cpp
--- a/Engine/Source/Developer/Horde/Private/Storage/Bundles/V2/PacketHandle.cpp
+++ b/Engine/Source/Developer/Horde/Private/Storage/Bundles/V2/PacketHandle.cpp
@@ -3,7 +3,7 @@
 #include "Storage/Bundles/V2/PacketHandle.h"
 #include "Storage/Bundles/V2/PacketReader.h"
 #include "Storage/Bundles/V2/ExportHandle.h"
-#include "../../../HordePlatform.h"
+#include "HordePlatform.h"
 
 // --------------------------------------------------------------------------------------------------------
 
diff --git a/Engine/Source/Developer/Horde/Private/Storage/Bundles/V2/PacketReader.cpp b/Engine/Source/Developer/Horde/Private/Storage/Bundles/V2/PacketReader.cpp
--- a/Engine/Source/Developer/Horde/Private/Storage/Bundles/V2/PacketReader.cpp
+++ b/Engine/Source/Developer/Horde/Private/Storage/Bundles/V2/PacketReader.cpp
@@ -4,7 +4,7 @@
 #include "Storage/Bundles/V2/ExportHandle.h"
 #include "Storage/Bundles/V2/PacketHandle.h"
 #include "Storage/Blob.h"
-#include "../../../HordePlatform.h"
+#include "HordePlatform.h"
 
 FPacketReader::FPacketReader(TSharedRef<FStorageClient> InStorageClient, FBlobHandle InBundleHandle, FPacketHandle InPacketHandle, FPacket InPacket)
 	: StorageClient(MoveTemp(InStorageClient))
diff --git a/Engine/Source/Developer/Horde/Private/Storage/Bundles/V2/PacketWriter.cpp b/Engine/Source/Developer/Horde/Private/Storage/Bundles/V2/PacketWriter.cpp
--- a/Engine/Source/Developer/Horde/Private/Storage/Bundles/V2/PacketWriter.cpp
+++ b/Engine/Source/Developer/Horde/Private/Storage/Bundles/V2/PacketWriter.cpp
@@ -4,7 +4,7 @@
 #include "Storage/Bundles/Bundle.h"
 #include "Storage/BlobWriter.h"
 #include "Memory/SharedBuffer.h"
-#include "../../../HordePlatform.h"
+#include "HordePlatform.h"
 
 FPacketWriter::FPacketWriter(FBlobHandle InBundleHandle, FBlobHandle InPacketHandle)
 	: BundleHandle(MoveTemp(InBundleHandle))
diff --git a/Engine/Source/Developer/Horde/Private/Storage/Clients/BundleStorageClient.cpp b/Engine/Source/Developer/Horde/Private/Storage/Clients/BundleStorageClient.cpp
--- a/Engine/Source/Developer/Horde/Private/Storage/Clients/BundleStorageClient.cpp
+++ b/Engine/Source/Developer/Horde/Private/Storage/Clients/BundleStorageClient.cpp
@@ -6,7 +6,7 @@
 #include "Storage/Bundles/V2/BundleWriter.h"
 #include "Storage/Bundles/V2/ExportHandle.h"
 #include "Storage/Bundles/V2/PacketHandle.h"
-#include "../../HordePlatform.h"
+#include "HordePlatform.h"
 
 // ---------------------------------------------------------------------------------
 
diff --git a/Engine/Source/Developer/Horde/Private/Storage/Clients/FileStorageClient.cpp b/Engine/Source/Developer/Horde/Private/Storage/Clients/FileStorageClient.cpp
--- a/Engine/Source/Developer/Horde/Private/Storage/Clients/FileStorageClient.cpp
+++ b/Engine/Source/Developer/Horde/Private/Storage/Clients/FileStorageClient.cpp
@@ -6,7 +6,7 @@
 #include <iostream>
 #include <sstream>
 #include <fstream>
-#include "../../HordePlatform.h"
+#include "HordePlatform.h"
 
 FFileStorageClient::FFileStorageClient(std::filesystem::path InRootDir)
 	: RootDir(std::move(InRootDir))
diff --git a/Engine/Source/Developer/Horde/Private/Storage/Clients/KeyValueStorageClient.cpp b/Engine/Source/Developer/Horde/Private/Storage/Clients/KeyValueStorageClient.cpp
--- a/Engine/Source/Developer/Horde/Private/Storage/Clients/KeyValueStorageClient.cpp
+++ b/Engine/Source/Developer/Horde/Private/Storage/Clients/KeyValueStorageClient.cpp
@@ -2,7 +2,7 @@
 
 #include "Storage/Clients/KeyValueStorageClient.h"
 #include "Storage/Blob.h"
-#include "../../HordePlatform.h"
+#include "HordePlatform.h"
 
 class FKeyValueStorageClient::FHandleData : public FBlobHandleData
 {
diff --git a/Engine/Source/Developer/Horde/Private/Storage/Nodes/DirectoryNode.cpp b/Engine/Source/Developer/Horde/Private/Storage/Nodes/DirectoryNode.cpp
--- a/Engine/Source/Developer/Horde/Private/Storage/Nodes/DirectoryNode.cpp
+++ b/Engine/Source/Developer/Horde/Private/Storage/Nodes/DirectoryNode.cpp
@@ -4,7 +4,7 @@
 #include "Storage/Blob.h"
 #include "Storage/BlobReader.h"
 #include "Storage/BlobWriter.h"
-#include "../../HordePlatform.h"
+#include "HordePlatform.h"
 
 const FBlobType FDirectoryNode::BlobType(FGuid(0x0714EC11, 0x4D07291A, 0x8AE77F86, 0x799980D6), 1);
 
diff --git a/Engine/Source/Developer/Horde/Private/Storage/Nodes/FileEntry.cpp b/Engine/Source/Developer/Horde/Private/Storage/Nodes/FileEntry.cpp
--- a/Engine/Source/Developer/Horde/Private/Storage/Nodes/FileEntry.cpp
+++ b/Engine/Source/Developer/Horde/Private/Storage/Nodes/FileEntry.cpp
@@ -4,7 +4,7 @@
 #include "IO/IoHash.h"
 #include "Storage/BlobReader.h"
 #include "Storage/BlobWriter.h"
-#include "../../HordePlatform.h"
+#include "HordePlatform.h"
 
 FFileEntry::FFileEntry(FBlobHandleWithHash InTarget, FUtf8String InName, EFileEntryFlags InFlags, int64 InLength, const FIoHash& InHash, FSharedBufferView InCustomData)
 	: Target(MoveTemp(InTarget))
diff --git a/Engine/Source/Developer/Horde/Public/Storage/Bundles/V2/Packet.h b/Engine/Source/Developer/Horde/Public/Storage/Bundles/V2/Packet.h
--- a/Engine/Source/Developer/Horde/Public/Storage/Bundles/V2/Packet.h
+++ b/Engine/Source/Developer/Horde/Public/Storage/Bundles/V2/Packet.h
@@ -4,9 +4,9 @@
 
 #include "CoreMinimal.h"
 #include "Containers/Utf8String.h"
+#include "SharedBufferView.h"
 #include "Storage/BlobType.h"
-#include "../BundleCompression.h"
-#include "../../../SharedBufferView.h"
+#include "Storage/Bundles/BundleCompression.h"
 
 class FChunkedBufferWriter;
 struct FPacketImport;
diff --git a/Engine/Source/Developer/Horde/Public/Storage/Nodes/ChunkNode.h b/Engine/Source/Developer/Horde/Public/Storage/Nodes/ChunkNode.h
--- a/Engine/Source/Developer/Horde/Public/Storage/Nodes/ChunkNode.h
+++ b/Engine/Source/Developer/Horde/Public/Storage/Nodes/ChunkNode.h
@@ -4,16 +4,17 @@
 
 #include "CoreMinimal.h"
 #include "IO/IoHash.h"
-#include "../BlobType.h"
-#include "../BlobHandle.h"
-#include "../BlobWriter.h"
-#include "../../SharedBufferView.h"
+#include "Storage/BlobHandle.h"
+#include "Storage/BlobType.h"
+#include "Storage/BlobWriter.h"
+#include "SharedBufferView.h"
 #include "Hash/Blake3.h"
 #include "Hash/BuzHash.h"
 
 /**
  * Options for chunking data
  */
+
 struct HORDE_API FChunkingOptions
 {
 	static const FChunkingOptions Default;
diff --git a/Engine/Source/Developer/Horde/Public/Storage/Nodes/FileEntry.h b/Engine/Source/Developer/Horde/Public/Storage/Nodes/FileEntry.h
--- a/Engine/Source/Developer/Horde/Public/Storage/Nodes/FileEntry.h
+++ b/Engine/Source/Developer/Horde/Public/Storage/Nodes/FileEntry.h
@@ -5,8 +5,9 @@
 #include "CoreMinimal.h"
 #include "Containers/Utf8String.h"
 #include "IO/IoHash.h"
-#include "../BlobHandle.h"
-#include "../../SharedBufferView.h"
+#include "SharedBufferView.h"
+#include "Storage/BlobHandle.h"
+
 
 class FBlobReader;
 class FBlobWriter;
