[UBA] Fix hang in UbaController when UbaHordeManager is about to exit while agent binaries are still being uploaded.

This adds an output parameter to the C++ interface of Horde channels to respond to timeouts.
The timeout flag is necessary to determine whether a null response during uploading the UbaAgent binaries are the result of an end-of-stream or a timeout.
This way we can check periodically (after each timeout) whether the uploading should be canceled, which is triggered by the UbaAgentManager exit event.

#rnx

diff --git a/Engine/Source/Developer/Horde/Private/Compute/AgentMessage.cpp b/Engine/Source/Developer/Horde/Private/Compute/AgentMessage.cpp
--- a/Engine/Source/Developer/Horde/Private/Compute/AgentMessage.cpp
+++ b/Engine/Source/Developer/Horde/Private/Compute/AgentMessage.cpp
@@ -153,7 +153,7 @@ void FAgentMessageChannel::Xor(const unsigned char* Data, size_t Length, unsigne
 	FlushMessage();
 }
 
-EAgentMessageType FAgentMessageChannel::ReadResponse(int32 TimeoutMS)
+EAgentMessageType FAgentMessageChannel::ReadResponse(int32 TimeoutMS, bool* bOutTimedOut)
 {
 	if (ResponseData)
 	{
@@ -162,14 +162,14 @@ EAgentMessageType FAgentMessageChannel::ReadResponse(int32 TimeoutMS)
 		ResponseLength = 0;
 	}
 
-	const unsigned char* Header = ChannelBuffers->Reader.WaitToRead(MessageHeaderLength, TimeoutMS);
+	const unsigned char* Header = ChannelBuffers->Reader.WaitToRead(MessageHeaderLength, TimeoutMS, bOutTimedOut);
 	if (Header == nullptr)
 	{
 		return EAgentMessageType::None;
 	}
 	unsigned int Length = *((unsigned int*)(Header + 1));
 
-	Header = ChannelBuffers->Reader.WaitToRead(MessageHeaderLength + Length);
+	Header = ChannelBuffers->Reader.WaitToRead(MessageHeaderLength + Length, TimeoutMS, bOutTimedOut);
 
 	ResponseType = (EAgentMessageType)Header[0];
 	ResponseData = Header + MessageHeaderLength;
diff --git a/Engine/Source/Developer/Horde/Private/Compute/ComputeBuffer.cpp b/Engine/Source/Developer/Horde/Private/Compute/ComputeBuffer.cpp
--- a/Engine/Source/Developer/Horde/Private/Compute/ComputeBuffer.cpp
+++ b/Engine/Source/Developer/Horde/Private/Compute/ComputeBuffer.cpp
@@ -617,7 +617,7 @@ struct FComputeBufferDetail
 		}
 	}
 
-	const unsigned char* WaitToRead(int ReaderIdx, size_t MinSize, int TimeoutMs)
+	const unsigned char* WaitToRead(int ReaderIdx, size_t MinSize, int TimeoutMs, bool* bOutHasTimedOut)
 	{
 		for (; ; )
 		{
@@ -636,6 +636,10 @@ struct FComputeBufferDetail
 				// Wait until the current chunk is readable
 				if (!ReaderEvents[ReaderIdx].Wait(TimeoutMs))
 				{
+					if (bOutHasTimedOut)
+					{
+						*bOutHasTimedOut = true;
+					}
 					return nullptr;
 				}
 			}
@@ -649,6 +653,10 @@ struct FComputeBufferDetail
 				// Wait until there is more data in the chunk
 				if (!ReaderEvents[ReaderIdx].Wait(TimeoutMs))
 				{
+					if (bOutHasTimedOut)
+					{
+						*bOutHasTimedOut = true;
+					}
 					return nullptr;
 				}
 			}
@@ -1039,9 +1047,9 @@ size_t FComputeBufferReader::GetMaxReadSize() const
 	return Detail->GetMaxReadSize(ReaderIdx);
 }
 
-size_t FComputeBufferReader::Read(void* Buffer, size_t MaxSize, int TimeoutMs)
+size_t FComputeBufferReader::Read(void* Buffer, size_t MaxSize, int TimeoutMs, bool* bOutTimedOut)
 {
-	const unsigned char* RecvData = WaitToRead(1, TimeoutMs);
+	const unsigned char* RecvData = WaitToRead(1, TimeoutMs, bOutTimedOut);
 	if (RecvData == nullptr)
 	{
 		return 0;
@@ -1058,10 +1066,10 @@ size_t FComputeBufferReader::Read(void* Buffer, size_t MaxSize, int TimeoutMs)
 	return RecvSize;
 }
 
-const unsigned char* FComputeBufferReader::WaitToRead(size_t MinSize, int TimeoutMs)
+const unsigned char* FComputeBufferReader::WaitToRead(size_t MinSize, int TimeoutMs, bool* bOutTimedOut)
 {
 	check(Detail);
-	return Detail->WaitToRead(ReaderIdx, MinSize, TimeoutMs);
+	return Detail->WaitToRead(ReaderIdx, MinSize, TimeoutMs, bOutTimedOut);
 }
 
 FComputeBufferReader::FComputeBufferReader(FComputeBufferDetail* InDetail, int InReaderIdx)
diff --git a/Engine/Source/Developer/Horde/Public/Compute/AgentMessage.h b/Engine/Source/Developer/Horde/Public/Compute/AgentMessage.h
--- a/Engine/Source/Developer/Horde/Public/Compute/AgentMessage.h
+++ b/Engine/Source/Developer/Horde/Public/Compute/AgentMessage.h
@@ -131,7 +131,7 @@ public:
 	//// Responses ////
 
 	// Reads a response from the remote. Other Read methods can be used to access response data.
-	HORDE_API EAgentMessageType ReadResponse(int32 TimeoutMS = -1);
+	HORDE_API EAgentMessageType ReadResponse(int32 TimeoutMS = -1, bool* bOutTimedOut = nullptr);
 
 	// Gets the raw data from a response
 	HORDE_API const void* GetResponseData() const { return ResponseData; }
diff --git a/Engine/Source/Developer/Horde/Public/Compute/ComputeBuffer.h b/Engine/Source/Developer/Horde/Public/Compute/ComputeBuffer.h
--- a/Engine/Source/Developer/Horde/Public/Compute/ComputeBuffer.h
+++ b/Engine/Source/Developer/Horde/Public/Compute/ComputeBuffer.h
@@ -109,11 +109,11 @@ public:
 	HORDE_API size_t GetMaxReadSize() const;
 
 	// Reads data into the given buffer
-	HORDE_API size_t Read(void* Buffer, size_t MaxSize, int TimeoutMs = -1);
+	HORDE_API size_t Read(void* Buffer, size_t MaxSize, int TimeoutMs = -1, bool* bOutTimedOut = nullptr);
 
 	// Waits until the given amount of data has been read, and returns a pointer to it. Returns nullptr if the timeout expires, or
 	// if the requested amount of data is not in a contiguous block of memory.
-	HORDE_API const unsigned char* WaitToRead(size_t MinSize, int TimeoutMs = -1);
+	HORDE_API const unsigned char* WaitToRead(size_t MinSize, int TimeoutMs = -1, bool* bOutTimedOut = nullptr);
 
 private:
 	struct FReaderRef;
diff --git a/Engine/Source/Developer/UbaCoordinatorHorde/Private/UbaHordeAgent.cpp b/Engine/Source/Developer/UbaCoordinatorHorde/Private/UbaHordeAgent.cpp
--- a/Engine/Source/Developer/UbaCoordinatorHorde/Private/UbaHordeAgent.cpp
+++ b/Engine/Source/Developer/UbaCoordinatorHorde/Private/UbaHordeAgent.cpp
@@ -100,7 +100,7 @@ bool FUbaHordeAgent::BeginCommunication()
 	return true;
 }
 
-bool FUbaHordeAgent::UploadBinaries(const FString& BundleDirectory, const char* BundleLocator)
+bool FUbaHordeAgent::UploadBinaries(const FString& BundleDirectory, const char* BundleLocator, FEvent* CancelEvent)
 {
 	// Prepare for giving the remote machine work.
 	// 1) Send binaries to remote machine.
@@ -144,12 +144,30 @@ bool FUbaHordeAgent::UploadBinaries(const FString& BundleDirectory, const char*
 
 	EAgentMessageType Type = EAgentMessageType::None;
 
+	constexpr int kReadResponseTimeoutMS = 1000;
+
 	for (;;)
 	{
-		Type = ChildChannel->ReadResponse();
+		// Try to read channel response and check for cancel event periodically
+		bool bTimedOut = false;
+		Type = ChildChannel->ReadResponse(kReadResponseTimeoutMS, &bTimedOut);
 		if (Type != EAgentMessageType::ReadBlob)
 		{
-			break;
+			if (bTimedOut)
+			{
+				// If the cancel event was triggered, stop waiting for channel messages
+				if (CancelEvent != nullptr && CancelEvent->Wait(0))
+				{
+					UE_LOG(LogUbaHordeAgent, Log, TEXT("Uploading UBA agent binaries has been canceled"));
+					return false;
+				}
+				continue;
+			}
+			else
+			{
+				// End of stream reached
+				break;
+			}
 		}
 
 		AgentMessage::FBlobRequest BlobRequest;
diff --git a/Engine/Source/Developer/UbaCoordinatorHorde/Private/UbaHordeAgent.h b/Engine/Source/Developer/UbaCoordinatorHorde/Private/UbaHordeAgent.h
--- a/Engine/Source/Developer/UbaCoordinatorHorde/Private/UbaHordeAgent.h
+++ b/Engine/Source/Developer/UbaCoordinatorHorde/Private/UbaHordeAgent.h
@@ -22,7 +22,7 @@ public:
 
 	/* Upload the proper executables/binaries. The Locator is just the contents of the bundle ref file.
 	 * e.g. Locator = abcdef1234567890abcdef1234567890abcdef12@34567890abcdef1234567890#123 */
-	bool UploadBinaries(const FString& BundleDirectory, const char* BundleLocator);
+	bool UploadBinaries(const FString& BundleDirectory, const char* BundleLocator, FEvent* CancelEvent = nullptr);
 
 	void Execute(const char* Exe, const char* const * Args, size_t NumArgs, const char* WorkingDir, const char* const * EnvVars, size_t NumEnvVars, bool bUseWine = false);
 
diff --git a/Engine/Source/Developer/UbaCoordinatorHorde/Private/UbaHordeAgentManager.cpp b/Engine/Source/Developer/UbaCoordinatorHorde/Private/UbaHordeAgentManager.cpp
--- a/Engine/Source/Developer/UbaCoordinatorHorde/Private/UbaHordeAgentManager.cpp
+++ b/Engine/Source/Developer/UbaCoordinatorHorde/Private/UbaHordeAgentManager.cpp
@@ -388,7 +388,7 @@ void FUbaHordeAgentManager::ThreadAgent(FHordeAgentWrapper& Wrapper)
 				return;
 			}
 
-			if (!Agent->UploadBinaries(BundleDirectory, reinterpret_cast<const char*>(Locator.GetData())))
+			if (!Agent->UploadBinaries(BundleDirectory, reinterpret_cast<const char*>(Locator.GetData()), &ShouldExit))
 			{
 				return;
 			}
