[UBA] Output frame header size when considered invalid instead of just 'check(false)' assertion.

#rnx

diff --git a/Engine/Source/Developer/Horde/Private/Compute/ComputeSocket.cpp b/Engine/Source/Developer/Horde/Private/Compute/ComputeSocket.cpp
--- a/Engine/Source/Developer/Horde/Private/Compute/ComputeSocket.cpp
+++ b/Engine/Source/Developer/Horde/Private/Compute/ComputeSocket.cpp
@@ -354,7 +354,7 @@ public:
 			}
 			else
 			{
-				check(false);
+				checkf(false, TEXT("Invalid frame header size received (%d)"), Header.Size);
 			}
 		}
 	}
