[Horde]
* Added more info to assert when failing to create semaphore

diff --git a/Engine/Source/Developer/Horde/Private/Compute/ComputePlatform.cpp b/Engine/Source/Developer/Horde/Private/Compute/ComputePlatform.cpp
--- a/Engine/Source/Developer/Horde/Private/Compute/ComputePlatform.cpp
+++ b/Engine/Source/Developer/Horde/Private/Compute/ComputePlatform.cpp
@@ -73,7 +73,7 @@ bool FComputeEvent::Create(const char* Name)
 		Handle = Value;
 		return true;
 	}
-	checkf(false, TEXT("sem_open failed (%hs)"), strerror(errno));
+	checkf(false, TEXT("sem_open to create semaphore with name %hs failed (%hs)"), Name, strerror(errno));
 	return false;
 #endif
 }
