[Horde]
* Encoded semaphore name to be shorter to fit on macos limitations

diff --git a/Engine/Source/Developer/Horde/Private/HordePlatform.cpp b/Engine/Source/Developer/Horde/Private/HordePlatform.cpp
--- a/Engine/Source/Developer/Horde/Private/HordePlatform.cpp
+++ b/Engine/Source/Developer/Horde/Private/HordePlatform.cpp
@@ -76,16 +76,39 @@ bool FHordePlatform::GetEnvironmentVariable(const char* Name, char* Buffer, size
 void FHordePlatform::CreateUniqueIdentifier(char* NameBuffer, size_t NameBufferLen)
 {
 	static int32 Counter = 0;
+	uint32 Index = (uint32)FPlatformAtomics::InterlockedIncrement(&Counter);
 
 #if UE_HORDE_PLATFORM_WINDOWS
-	DWORD Pid = GetCurrentProcessId();
-	ULONGLONG TickCount = GetTickCount64();
-	snprintf(NameBuffer, NameBufferLen, "%lu_%llu_%lu", Pid, TickCount, (unsigned long)FPlatformAtomics::InterlockedIncrement(&Counter));
+	uint32 Pid = GetCurrentProcessId();
+	uint64 Time = GetTickCount64();
 #else
 	struct timespec ts;
 	clock_gettime(CLOCK_MONOTONIC, &ts);
-	snprintf(NameBuffer, NameBufferLen, "%u%zu%zu_%d", getpid(), (size_t)(ts.tv_sec%100000), (size_t)ts.tv_nsec, FPlatformAtomics::InterlockedIncrement(&Counter));
+	uint32 Pid = (uint32)getpid();
+	uint64 Time = (uint64)(ts.tv_sec%100000);
 #endif
+
+	TStringBuilderBase<char> Builder(NameBuffer, NameBufferLen);
+
+	// Encode to be as short as possible.. 
+	auto AppendValue = [&](auto num)
+	{
+		constexpr char charset[] = "0123456789ABCDEFGHIJKLMNOPQRSTUVWXYZabcdefghijklmnopqrstuvwxyz";
+		constexpr int base = sizeof(charset);
+		do
+		{
+			Builder.AppendChar(charset[num % base]);
+			num /= base;
+		}
+		while (num);
+	};
+
+	AppendValue(Pid);
+	Builder.AppendChar('_');
+	AppendValue(Time);
+	Builder.AppendChar('_');
+	AppendValue(Index);
+	Builder.ToString(); // To ensure terminated
 }
 
 void FHordePlatform::CreateUniqueName(char* NameBuffer, size_t NameBufferLen)
