Horde: Add some info to an assert

diff --git a/Engine/Source/Developer/Horde/Private/Compute/ComputeBuffer.cpp b/Engine/Source/Developer/Horde/Private/Compute/ComputeBuffer.cpp
--- a/Engine/Source/Developer/Horde/Private/Compute/ComputeBuffer.cpp
+++ b/Engine/Source/Developer/Horde/Private/Compute/ComputeBuffer.cpp
@@ -503,7 +503,8 @@ struct FComputeBufferDetail
 		for (; ; )
 		{
 			FWriterState WriterState = WriterStatePtr.Get();
-			check(WriterState.GetRefCount() == 0);
+			int WriterStateRefCount = WriterState.GetRefCount();
+			checkf(WriterStateRefCount == 0, TEXT("Attempting to CreateWriter with a non-zero WriterState refcount: %d"), WriterStateRefCount);
 
 			if (WriterStatePtr.TryUpdate(WriterState, FWriterState(WriterState.GetChunkIdx(), WriterState.GetReaderFlags(), 1, WriterState.HasWrapped())))
 			{
