[Horde]
* Speculative revert of logic calculating unique identifiers. Suspect it is causing issues on windows

diff --git a/Engine/Source/Developer/Horde/Private/HordePlatform.cpp b/Engine/Source/Developer/Horde/Private/HordePlatform.cpp
--- a/Engine/Source/Developer/Horde/Private/HordePlatform.cpp
+++ b/Engine/Source/Developer/Horde/Private/HordePlatform.cpp
@@ -76,17 +76,18 @@ bool FHordePlatform::GetEnvironmentVariable(const char* Name, char* Buffer, size
 void FHordePlatform::CreateUniqueIdentifier(char* NameBuffer, size_t NameBufferLen)
 {
 	static int32 Counter = 0;
-	uint32 Index = (uint32)FPlatformAtomics::InterlockedIncrement(&Counter);
 
 #if UE_HORDE_PLATFORM_WINDOWS
-	uint32 Pid = GetCurrentProcessId();
-	uint64 Time = GetTickCount64();
+	DWORD Pid = GetCurrentProcessId();
+	ULONGLONG TickCount = GetTickCount64();
+	snprintf(NameBuffer, NameBufferLen, "%lu_%llu_%lu", Pid, TickCount, (unsigned long)FPlatformAtomics::InterlockedIncrement(&Counter));
 #else
+	uint32 Index = (uint32)FPlatformAtomics::InterlockedIncrement(&Counter);
+
 	struct timespec ts;
 	clock_gettime(CLOCK_MONOTONIC, &ts);
 	uint32 Pid = (uint32)getpid();
 	uint64 Time = (uint64)(ts.tv_sec%100000);
-#endif
 
 	TStringBuilderBase<char> Builder(NameBuffer, NameBufferLen);
 
@@ -109,6 +110,7 @@ void FHordePlatform::CreateUniqueIdentifier(char* NameBuffer, size_t NameBufferL
 	Builder.AppendChar('_');
 	AppendValue(Index);
 	Builder.ToString(); // To ensure terminated
+#endif
 }
 
 void FHordePlatform::CreateUniqueName(char* NameBuffer, size_t NameBufferLen)
