[Horde]
* Resubmitted change that creates a shorter unique identifier now when we found the bug where it could write in terminator character too early

diff --git a/Engine/Source/Developer/Horde/Private/HordePlatform.cpp b/Engine/Source/Developer/Horde/Private/HordePlatform.cpp
--- a/Engine/Source/Developer/Horde/Private/HordePlatform.cpp
+++ b/Engine/Source/Developer/Horde/Private/HordePlatform.cpp
@@ -76,32 +76,31 @@ bool FHordePlatform::GetEnvironmentVariable(const char* Name, char* Buffer, size
 void FHordePlatform::CreateUniqueIdentifier(char* NameBuffer, size_t NameBufferLen)
 {
 	static int32 Counter = 0;
+	uint32 Index = (uint32)FPlatformAtomics::InterlockedIncrement(&Counter);
 
 #if UE_HORDE_PLATFORM_WINDOWS
-	DWORD Pid = GetCurrentProcessId();
-	ULONGLONG TickCount = GetTickCount64();
-	snprintf(NameBuffer, NameBufferLen, "%lu_%llu_%lu", Pid, TickCount, (unsigned long)FPlatformAtomics::InterlockedIncrement(&Counter));
+	uint32 Pid = GetCurrentProcessId();
+	uint64 Time = GetTickCount64();
 #else
-	uint32 Index = (uint32)FPlatformAtomics::InterlockedIncrement(&Counter);
-
 	struct timespec ts;
 	clock_gettime(CLOCK_MONOTONIC, &ts);
 	uint32 Pid = (uint32)getpid();
 	uint64 Time = (uint64)(ts.tv_sec%100000);
+#endif
 
 	TStringBuilderBase<char> Builder(NameBuffer, NameBufferLen);
 
 	// Encode to be as short as possible.. 
-	auto AppendValue = [&](auto num)
+	auto AppendValue = [&](auto Value)
 	{
-		constexpr char charset[] = "0123456789ABCDEFGHIJKLMNOPQRSTUVWXYZabcdefghijklmnopqrstuvwxyz";
-		constexpr int base = sizeof(charset);
+		constexpr char CharSet[] = "0123456789ABCDEFGHIJKLMNOPQRSTUVWXYZabcdefghijklmnopqrstuvwxyz";
+		constexpr int Base = sizeof(CharSet) - 1;
 		do
 		{
-			Builder.AppendChar(charset[num % base]);
-			num /= base;
+			Builder.AppendChar(CharSet[Value % Base]);
+			Value /= Base;
 		}
-		while (num);
+		while (Value);
 	};
 
 	AppendValue(Pid);
@@ -110,7 +109,6 @@ void FHordePlatform::CreateUniqueIdentifier(char* NameBuffer, size_t NameBufferL
 	Builder.AppendChar('_');
 	AppendValue(Index);
 	Builder.ToString(); // To ensure terminated
-#endif
 }
 
 void FHordePlatform::CreateUniqueName(char* NameBuffer, size_t NameBufferLen)
