[Horde]
* Added support for creating anonymous FComputeBuffers and changed so uba with socket based horde transport use anonymous compute buffers. Reason is to prevent leakages of named objects on linux/mac which eventually will run out of them and require a hard restart of the machine

diff --git a/Engine/Source/Developer/Horde/Private/Compute/ComputeBuffer.cpp b/Engine/Source/Developer/Horde/Private/Compute/ComputeBuffer.cpp
--- a/Engine/Source/Developer/Horde/Private/Compute/ComputeBuffer.cpp
+++ b/Engine/Source/Developer/Horde/Private/Compute/ComputeBuffer.cpp
@@ -352,7 +352,7 @@ struct FComputeBufferDetail
 		snprintf(NameBuffer, FComputeBuffer::MaxNameLength, "%s_M", Name);
 
 		TUniquePtr<FComputeBufferDetail> Detail = MakeUnique<FComputeBufferDetail>(Name);
-		if (!Detail->MemoryMappedFile.Create(NameBuffer, Capacity))
+		if (!Detail->MemoryMappedFile.Create(Params.Anonymous ? nullptr : NameBuffer, Capacity))
 		{
 			checkf(false, TEXT("Failed to create memory mapped file %hs"), NameBuffer);
 			return nullptr;
@@ -368,7 +368,7 @@ struct FComputeBufferDetail
 		Detail->Header = new(Pointer) FHeader(Params.NumReaders, Params.NumChunks, Params.ChunkLength);
 
 		snprintf(NameBuffer, FComputeBuffer::MaxNameLength, "%s_W", Name);
-		if (!Detail->WriterEvent.Create(NameBuffer))
+		if (!Detail->WriterEvent.Create(Params.Anonymous ? nullptr : NameBuffer))
 		{
 			checkf(false, TEXT("Failed to create writer event %hs in memory mapping"), NameBuffer);
 			return nullptr;
@@ -377,7 +377,7 @@ struct FComputeBufferDetail
 		for (int ReaderIdx = 0; ReaderIdx < Detail->Header->NumReaders; ReaderIdx++)
 		{
 			snprintf(NameBuffer, FComputeBuffer::MaxNameLength, "%s_R%d", Name, ReaderIdx);
-			if (!Detail->ReaderEvents[ReaderIdx].Create(NameBuffer))
+			if (!Detail->ReaderEvents[ReaderIdx].Create(Params.Anonymous ? nullptr : NameBuffer))
 			{
 				checkf(false, TEXT("Failed to create reader event %hs in memory mapping"), NameBuffer);
 				return nullptr;
diff --git a/Engine/Source/Developer/Horde/Private/Compute/ComputePlatform.cpp b/Engine/Source/Developer/Horde/Private/Compute/ComputePlatform.cpp
--- a/Engine/Source/Developer/Horde/Private/Compute/ComputePlatform.cpp
+++ b/Engine/Source/Developer/Horde/Private/Compute/ComputePlatform.cpp
@@ -1,6 +1,7 @@
 // Copyright Epic Games, Inc. All Rights Reserved.
 
 #include "Compute/ComputePlatform.h"
+#include "HAL/Event.h"
 #include <assert.h>
 #include <wchar.h>
 #include <bit>
@@ -51,6 +52,7 @@
 
 FComputeEvent::FComputeEvent()
 	: Handle(nullptr)
+	, Anonymous(false)
 {
 }
 
@@ -63,6 +65,15 @@ bool FComputeEvent::Create(const char* Name)
 {
 	Close();
 
+	Anonymous = Name == nullptr;
+
+	if (Anonymous)
+	{
+		new (&Handle) FEventRef();
+		return true;
+	}
+
+
 #if UE_COMPUTE_PLATFORM_WINDOWS
 	Handle = CreateEventA(NULL, false, false, Name);
 	return Handle != nullptr;
@@ -82,6 +93,9 @@ bool FComputeEvent::OpenExisting(const char* Name)
 {
 	Close();
 
+	check(Name);
+	Anonymous = false;
+
 #if UE_COMPUTE_PLATFORM_WINDOWS
 	Handle = OpenEventA(SYNCHRONIZE | EVENT_MODIFY_STATE, true, Name);
 	return Handle != nullptr;
@@ -98,6 +112,13 @@ bool FComputeEvent::OpenExisting(const char* Name)
 
 void FComputeEvent::Close()
 {
+	if (Anonymous)
+	{
+		((FEventRef*)&Handle)->~FEventRef();
+		Handle = nullptr;
+		return;
+	}
+
 #if UE_COMPUTE_PLATFORM_WINDOWS
 	if (Handle != nullptr)
 	{
@@ -115,6 +136,12 @@ void FComputeEvent::Close()
 
 void FComputeEvent::Signal()
 {
+	if (Anonymous)
+	{
+		(*((FEventRef*)&Handle))->Trigger();
+		return;
+	}
+
 #if UE_COMPUTE_PLATFORM_WINDOWS
 	SetEvent(Handle);
 #else
@@ -124,6 +151,11 @@ void FComputeEvent::Signal()
 
 bool FComputeEvent::Wait(int timeoutMs)
 {
+	if (Anonymous)
+	{
+		return (*((FEventRef*)&Handle))->Wait(timeoutMs == -1 ? MAX_uint32 : timeoutMs);
+	}
+
 #if UE_COMPUTE_PLATFORM_WINDOWS
 	DWORD WaitParam = (timeoutMs < 0) ? INFINITE : (DWORD)timeoutMs;
 	return WaitForSingleObject(Handle, WaitParam) != WAIT_TIMEOUT;
@@ -189,6 +221,7 @@ FComputeMemoryMappedFile::FComputeMemoryMappedFile()
 	, Pointer(nullptr)
 	, MappedSize(0)
 	, OwnerName(nullptr)
+	, Anonymous(false)
 {
 }
 
@@ -201,6 +234,15 @@ bool FComputeMemoryMappedFile::Create(const char* Name, long long Capacity)
 {
 	Close();
 
+	Anonymous = Name == nullptr;
+
+	if (Anonymous)
+	{
+		Pointer = FMemory::Malloc(Capacity);
+		MappedSize = Capacity;
+		return true;
+	}
+
 #if UE_COMPUTE_PLATFORM_WINDOWS
 	LARGE_INTEGER LargeInteger;
 	LargeInteger.QuadPart = Capacity;
@@ -251,6 +293,9 @@ bool FComputeMemoryMappedFile::OpenExisting(const char* Name)
 {
 	Close();
 
+	check(Name);
+	Anonymous = false;
+
 #if UE_COMPUTE_PLATFORM_WINDOWS
 	Handle = OpenFileMappingA(FILE_MAP_ALL_ACCESS, true, Name);
 	if (Handle == nullptr)
@@ -292,6 +337,13 @@ bool FComputeMemoryMappedFile::OpenExisting(const char* Name)
 
 void FComputeMemoryMappedFile::Close()
 {
+	if (Anonymous)
+	{
+		FMemory::Free(Pointer);
+		Pointer = nullptr;
+		return;
+	}
+
 #if UE_COMPUTE_PLATFORM_WINDOWS
 	if (Pointer != nullptr)
 	{
@@ -308,6 +360,7 @@ void FComputeMemoryMappedFile::Close()
 	if(OwnerName != nullptr)
 	{
 		shm_unlink(OwnerName);
+		free(OwnerName);
 		OwnerName = nullptr;
 	}
 	
diff --git a/Engine/Source/Developer/Horde/Private/Compute/ComputePlatform.h b/Engine/Source/Developer/Horde/Private/Compute/ComputePlatform.h
--- a/Engine/Source/Developer/Horde/Private/Compute/ComputePlatform.h
+++ b/Engine/Source/Developer/Horde/Private/Compute/ComputePlatform.h
@@ -33,6 +33,7 @@ public:
 
 private:
 	void* Handle;
+	bool Anonymous;
 };
 
 //
@@ -61,4 +62,5 @@ private:
 	void* Pointer;
 	long long MappedSize;
 	char* OwnerName;
+	bool Anonymous;
 };
diff --git a/Engine/Source/Developer/Horde/Private/Compute/ComputeSocket.cpp b/Engine/Source/Developer/Horde/Private/Compute/ComputeSocket.cpp
--- a/Engine/Source/Developer/Horde/Private/Compute/ComputeSocket.cpp
+++ b/Engine/Source/Developer/Horde/Private/Compute/ComputeSocket.cpp
@@ -22,16 +22,19 @@ FComputeSocket::~FComputeSocket()
 {
 }
 
-TSharedPtr<FComputeChannel> FComputeSocket::CreateChannel(int ChannelId)
+TSharedPtr<FComputeChannel> FComputeSocket::CreateChannel(int ChannelId, bool Anonymous)
 {
+	FComputeBuffer::FParams Params;
+	Params.Anonymous = Anonymous;
+
 	FComputeBuffer RecvBuffer;
-	if (!RecvBuffer.CreateNew(FComputeBuffer::FParams()))
+	if (!RecvBuffer.CreateNew(Params))
 	{
 		return {};
 	}
 
 	FComputeBuffer SendBuffer;
-	if (!SendBuffer.CreateNew(FComputeBuffer::FParams()))
+	if (!SendBuffer.CreateNew(Params))
 	{
 		return {};
 	}
diff --git a/Engine/Source/Developer/Horde/Public/Compute/ComputeBuffer.h b/Engine/Source/Developer/Horde/Public/Compute/ComputeBuffer.h
--- a/Engine/Source/Developer/Horde/Public/Compute/ComputeBuffer.h
+++ b/Engine/Source/Developer/Horde/Public/Compute/ComputeBuffer.h
@@ -40,6 +40,7 @@ public:
 		int NumChunks = 2;
 		int ChunkLength = 512 * 1024;
 		int NumReaders = 1;
+		bool Anonymous = false; // Set this to true to use anonymous buffers. This could be used together with for example socket transports
 	};
 
 	HORDE_API FComputeBuffer();
diff --git a/Engine/Source/Developer/Horde/Public/Compute/ComputeSocket.h b/Engine/Source/Developer/Horde/Public/Compute/ComputeSocket.h
--- a/Engine/Source/Developer/Horde/Public/Compute/ComputeSocket.h
+++ b/Engine/Source/Developer/Horde/Public/Compute/ComputeSocket.h
@@ -31,7 +31,7 @@ public:
 	HORDE_API virtual void AttachSendBuffer(int ChannelId, FComputeBuffer SendBuffer) = 0;
 
 	// Attaches a channel to this socket
-	HORDE_API TSharedPtr<FComputeChannel> CreateChannel(int ChannelId);
+	HORDE_API TSharedPtr<FComputeChannel> CreateChannel(int ChannelId, bool Anonymous = false);
 
 	// Attaches a channel to this socket
 	HORDE_API TSharedPtr<FComputeChannel> CreateChannel(int ChannelId, FComputeBuffer RecvBuffer, FComputeBuffer SendBuffer);
diff --git a/Engine/Source/Developer/UbaCoordinatorHorde/Private/UbaHordeAgent.cpp b/Engine/Source/Developer/UbaCoordinatorHorde/Private/UbaHordeAgent.cpp
--- a/Engine/Source/Developer/UbaCoordinatorHorde/Private/UbaHordeAgent.cpp
+++ b/Engine/Source/Developer/UbaCoordinatorHorde/Private/UbaHordeAgent.cpp
@@ -41,9 +41,10 @@ FUbaHordeAgent::FUbaHordeAgent(const FHordeRemoteMachineInfo& InMachineInfo)
 	// Create the compute socket and initialize a compute channel with the recv/send buffers.
 	// This will allow us to create the agent channel object which will allow us to communicate
 	// with the agent directly.
+	bool Anonymous = true; // There is no point wasting resources creating named events and file mappings if we use sockets for transport
 	HordeComputeSocket = CreateComputeSocket(MoveTemp(ComputeTransport), EComputeSocketEndpoint::Remote);
-	TSharedPtr<FComputeChannel> ComputeChannel = HordeComputeSocket->CreateChannel(0);
-	TSharedPtr<FComputeChannel> SecondComputeChannel = HordeComputeSocket->CreateChannel(100);
+	TSharedPtr<FComputeChannel> ComputeChannel = HordeComputeSocket->CreateChannel(0, Anonymous);
+	TSharedPtr<FComputeChannel> SecondComputeChannel = HordeComputeSocket->CreateChannel(100, Anonymous);
 
 	if (!ComputeChannel.IsValid() || !SecondComputeChannel.IsValid())
 	{
