[Horde] Improve error handling on failed messages receival.

- Add log category LogHorde to Horde module.
- Instead of crashing in the receive thread procedure, report a warning and exit the thread so that missing message content can be handled in the upper abstraction layers.

#rnx

diff --git a/Engine/Source/Developer/Horde/Private/Compute/ComputeSocket.cpp b/Engine/Source/Developer/Horde/Private/Compute/ComputeSocket.cpp
--- a/Engine/Source/Developer/Horde/Private/Compute/ComputeSocket.cpp
+++ b/Engine/Source/Developer/Horde/Private/Compute/ComputeSocket.cpp
@@ -349,7 +349,11 @@ public:
 		{
 			if (Header.Size >= 0)
 			{
-				ReadFrame(CachedWriters, Header.Channel, Header.Size);
+				if (!ReadFrame(CachedWriters, Header.Channel, Header.Size))
+				{
+					UE_LOG(LogHorde, Warning, TEXT("Failed to read frame header (Channel %d, Size %d)"), Header.Channel, Header.Size);
+					return;
+				}
 			}
 			else if (Header.Size == (int)EControlMessageType::Detach)
 			{
@@ -357,7 +361,8 @@ public:
 			}
 			else
 			{
-				checkf(false, TEXT("Invalid frame header size received (%d)"), Header.Size);
+				UE_LOG(LogHorde, Warning, TEXT("Invalid frame header size received (%d)"), Header.Size);
+				return;
 			}
 		}
 	}
diff --git a/Engine/Source/Developer/Horde/Private/Horde.cpp b/Engine/Source/Developer/Horde/Private/Horde.cpp
--- a/Engine/Source/Developer/Horde/Private/Horde.cpp
+++ b/Engine/Source/Developer/Horde/Private/Horde.cpp
@@ -1,12 +1,17 @@
 // Copyright Epic Games, Inc. All Rights Reserved.
 
 #include "Horde.h"
+
+#include "DesktopPlatformModule.h"
+#include "HordePlatform.h"
+#include "IDesktopPlatform.h"
 #include "Misc/CommandLine.h"
 #include "Misc/Parse.h"
 #include "Modules/ModuleManager.h"
 #include "Modules/ModuleInterface.h"
-#include "IDesktopPlatform.h"
-#include "DesktopPlatformModule.h"
+
+
+DEFINE_LOG_CATEGORY(LogHorde);
 
 struct FHordeModule : IModuleInterface
 {
diff --git a/Engine/Source/Developer/Horde/Private/HordePlatform.h b/Engine/Source/Developer/Horde/Private/HordePlatform.h
--- a/Engine/Source/Developer/Horde/Private/HordePlatform.h
+++ b/Engine/Source/Developer/Horde/Private/HordePlatform.h
@@ -4,6 +4,8 @@
 
 #include "CoreMinimal.h"
 
+DECLARE_LOG_CATEGORY_EXTERN(LogHorde, Log, All);
+
 //
 // Platform-specific utility functions
 //
