[UBA] Downgrade verbosity of failed frame header in ComputeSocket from Warning to Log.

#rnx

diff --git a/Engine/Source/Developer/Horde/Private/Compute/ComputeSocket.cpp b/Engine/Source/Developer/Horde/Private/Compute/ComputeSocket.cpp
--- a/Engine/Source/Developer/Horde/Private/Compute/ComputeSocket.cpp
+++ b/Engine/Source/Developer/Horde/Private/Compute/ComputeSocket.cpp
@@ -351,7 +351,7 @@ public:
 			{
 				if (!ReadFrame(CachedWriters, Header.Channel, Header.Size))
 				{
-					UE_LOG(LogHorde, Warning, TEXT("Failed to read frame header (Channel %d, Size %d)"), Header.Channel, Header.Size);
+					UE_LOG(LogHorde, Log, TEXT("Failed to read frame header (Channel %d, Size %d)"), Header.Channel, Header.Size);
 					return;
 				}
 			}
